[WaveformEditor] Move to main engine plugins folder for 5.1 release

diff --git a/Engine/Plugins/Editor/WaveformEditor/Config/BaseWaveformEditor.ini b/Engine/Plugins/Editor/WaveformEditor/Config/BaseWaveformEditor.ini
new file mode 100644
--- /dev/null
+++ b/Engine/Plugins/Editor/WaveformEditor/Config/BaseWaveformEditor.ini
@@ -0,0 +1,11 @@
+[/Script/WaveformEditorWidgets.WaveformEditorWidgetsSettings]
+PlayheadColor=(R=1.000000,G=0.100000,B=0.200000,A=1.000000)
+WaveformColor=(R=1.000000,G=1.000000,B=1.000000,A=1.000000)
+WaveformBackgroundColor=(R=0.017642,G=0.017642,B=0.017642,A=1.000000)
+MajorGridColor=(R=0.000000,G=0.000000,B=0.000000,A=1.000000)
+MinorGridColor=(R=0.000000,G=0.000000,B=0.000000,A=0.500000)
+RulerBackgroundColor=(R=0.000000,G=0.000000,B=0.000000,A=1.000000)
+RulerTicksColor=(R=1.000000,G=1.000000,B=1.000000,A=0.900000)
+RulerTextColor=(R=1.000000,G=1.000000,B=1.000000,A=0.900000)
+RulerFontSize=10.000000
+
diff --git a/Engine/Plugins/Editor/WaveformEditor/Source/WaveformEditor/Private/WaveformEditor.cpp b/Engine/Plugins/Editor/WaveformEditor/Source/WaveformEditor/Private/WaveformEditor.cpp
new file mode 100644
--- /dev/null
+++ b/Engine/Plugins/Editor/WaveformEditor/Source/WaveformEditor/Private/WaveformEditor.cpp
@@ -0,0 +1,545 @@
+// Copyright Epic Games, Inc. All Rights Reserved.
+
+#include "WaveformEditor.h"
+
+#include "AudioDevice.h"
+#include "Components/AudioComponent.h"
+#include "Sound/SoundWave.h"
+#include "Styling/AppStyle.h"
+#include "SWaveformPanel.h"
+#include "SWaveformTransformationsOverlay.h"
+#include "ToolMenus.h"
+#include "WaveformEditorCommands.h"
+#include "WaveformEditorLog.h"
+#include "WaveformEditorRenderData.h"
+#include "WaveformEditorStyle.h"
+#include "WaveformTransformationsRenderManager.h"
+#include "WaveformEditorWaveWriter.h"
+#include "Widgets/Docking/SDockTab.h"
+
+#define LOCTEXT_NAMESPACE "WaveformEditor"
+
+const FName FWaveformEditor::AppIdentifier("WaveformEditorApp");
+const FName FWaveformEditor::PropertiesTabId("WaveformEditor_Properties");
+const FName FWaveformEditor::WaveformDisplayTabId("WaveformEditor_Display");
+const FName FWaveformEditor::EditorName("Waveform Editor");
+const FName FWaveformEditor::ToolkitFName("WaveformEditor");
+
+bool FWaveformEditor::Init(const EToolkitMode::Type Mode, const TSharedPtr<IToolkitHost>& InitToolkitHost, USoundWave* SoundWaveToEdit)
+{
+	checkf(SoundWaveToEdit, TEXT("Tried to open a Soundwave Editor from a null soundwave"));
+
+	SoundWave = SoundWaveToEdit;
+
+	bool bIsInitialized = true;
+	
+	bIsInitialized &= SetUpPropertiesView();
+	bIsInitialized &= SetUpWaveformPanel();
+	bIsInitialized &= SetupAudioComponent();
+	bIsInitialized &= SetUpTransportController();
+	bIsInitialized &= SetUpWaveWriter();
+	bIsInitialized &= BindDelegates();
+
+	bIsInitialized &= RegisterToolbar();
+	bIsInitialized &= BindCommands();
+
+	const TSharedRef<FTabManager::FLayout>  StandaloneDefaultLayout = SetupStandaloneLayout();
+
+	if (bIsInitialized)
+	{
+		const bool bCreateDefaultStandaloneMenu = true;
+		const bool bCreateDefaultToolbar = true;
+		const bool bToolbarFocusable = false;
+		const bool bUseSmallIcons = true;
+
+		FAssetEditorToolkit::InitAssetEditor(
+			Mode,
+			InitToolkitHost,
+			AppIdentifier,
+			StandaloneDefaultLayout,
+			bCreateDefaultStandaloneMenu,
+			bCreateDefaultToolbar,
+			SoundWaveToEdit,
+			bToolbarFocusable,
+			bUseSmallIcons);
+	}
+
+	return bIsInitialized;
+}
+
+bool FWaveformEditor::SetupAudioComponent()
+{
+	if (SoundWave == nullptr)
+	{
+		return false;
+	}
+
+	if (AudioComponent == nullptr)
+	{
+		if (FAudioDeviceManager* AudioDeviceManager = FAudioDeviceManager::Get())
+		{
+			if (FAudioDevice* AudioDevice = AudioDeviceManager->GetMainAudioDeviceRaw())
+			{
+				USoundBase* SoundBase = Cast<USoundBase>(SoundWave);
+				AudioComponent = FAudioDevice::CreateComponent(SoundBase);
+			}
+		}
+	}
+
+	AudioComponent->bAutoDestroy = false;
+	AudioComponent->bIsUISound = true;
+	AudioComponent->bAllowSpatialization = false;
+	AudioComponent->bReverb = false;
+	AudioComponent->bCenterChannelOnly = false;
+	AudioComponent->bIsPreviewSound = true;
+
+	return AudioComponent != nullptr;
+}
+
+bool FWaveformEditor::SetUpTransportController()
+{
+	if (AudioComponent == nullptr)
+	{
+		UE_LOG(LogWaveformEditor, Warning, TEXT("Trying to setup transport controls with a null audio component"));
+		return false;
+	}
+
+	TransportController = MakeShared<FWaveformEditorTransportController>(AudioComponent);
+	return TransportController != nullptr;
+}
+
+bool FWaveformEditor::SetUpZoom()
+{
+	ZoomManager = MakeShared<FWaveformEditorZoomController>();
+	return ZoomManager != nullptr;
+}
+
+bool FWaveformEditor::BindDelegates()
+{
+	if (AudioComponent == nullptr)
+	{
+		UE_LOG(LogWaveformEditor, Warning, TEXT("Failed to bind to playback percentage change, audio component is null"));
+		return false;
+	}
+
+	AudioComponent->OnAudioPlaybackPercentNative.AddSP(this, &FWaveformEditor::HandlePlaybackPercentageChange);
+	AudioComponent->OnAudioPlayStateChangedNative.AddSP(this, &FWaveformEditor::HandleAudioComponentPlayStateChanged);
+	TransportCoordinator->OnPlayheadScrubUpdate.AddSP(this, &FWaveformEditor::HandlePlayheadScrub);
+	return true;
+}
+
+void FWaveformEditor::RegisterTabSpawners(const TSharedRef<FTabManager>& InTabManager)
+{
+	WorkspaceMenuCategory = InTabManager->AddLocalWorkspaceMenuCategory(LOCTEXT("WorkspaceMenu_WaveformEditor", "Sound Wave Editor"));
+
+	FAssetEditorToolkit::RegisterTabSpawners(InTabManager);
+
+	InTabManager->RegisterTabSpawner(PropertiesTabId, FOnSpawnTab::CreateSP(this, &FWaveformEditor::SpawnTab_Properties))
+		.SetDisplayName(LOCTEXT("DetailsTab", "Details"))
+		.SetGroup(WorkspaceMenuCategory.ToSharedRef())
+		.SetIcon(FSlateIcon(FAppStyle::GetAppStyleSetName(), "LevelEditor.Tabs.Details"));
+
+	InTabManager->RegisterTabSpawner(WaveformDisplayTabId, FOnSpawnTab::CreateSP(this, &FWaveformEditor::SpawnTab_WaveformDisplay))
+		.SetDisplayName(LOCTEXT("WaveformDisplayTab", "WaveformDisplay"))
+		.SetGroup(WorkspaceMenuCategory.ToSharedRef())
+		.SetIcon(FSlateIcon(FAppStyle::GetAppStyleSetName(), "LevelEditor.Tabs.Details"));
+
+}
+
+void FWaveformEditor::UnregisterTabSpawners(const TSharedRef<FTabManager>& InTabManager)
+{
+	InTabManager->UnregisterTabSpawner(PropertiesTabId);
+	InTabManager->UnregisterTabSpawner(WaveformDisplayTabId);
+}
+
+bool FWaveformEditor::RegisterToolbar()
+{
+	const FName MenuName = FAssetEditorToolkit::GetToolMenuToolbarName();
+
+	if (!UToolMenus::Get()->IsMenuRegistered(MenuName))
+	{
+		const FWaveformEditorCommands& Commands = FWaveformEditorCommands::Get();
+		UToolMenu* ToolBar = UToolMenus::Get()->RegisterMenu(MenuName, "AssetEditor.DefaultToolBar", EMultiBoxType::ToolBar);
+
+		if (ToolBar == nullptr)
+		{
+			return false;
+		}
+
+		FToolMenuInsert InsertAfterAssetSection("Asset", EToolMenuInsertType::After);
+		FToolMenuSection& PlayBackSection = ToolBar->AddSection("Transport Controls", TAttribute<FText>(), InsertAfterAssetSection);
+
+		FToolMenuEntry PlayEntry = FToolMenuEntry::InitToolBarButton(
+			Commands.PlaySoundWave,
+			LOCTEXT("WaveformEditorPlayButton", ""),
+			LOCTEXT("WaveformEditorPlayButtonTooltip", "Plays this SoundWave"),
+			FSlateIcon(FAppStyle::GetAppStyleSetName(), "PlayWorld.PlayInViewport")
+		);
+
+		PlayEntry.StyleNameOverride = FName("Toolbar.BackplateLeftPlay");
+		
+		FToolMenuEntry PauseEntry = FToolMenuEntry::InitToolBarButton(
+			Commands.PauseSoundWave,
+			LOCTEXT("WaveformEditorPauseButton", ""),
+			LOCTEXT("WaveformEditorPauseButtonTooltip", "Pauses this SoundWave"),
+			FSlateIcon(FAppStyle::GetAppStyleSetName(), "PlayWorld.PausePlaySession.Small")
+		);
+
+		PauseEntry.StyleNameOverride = FName("Toolbar.BackplateCenter");
+
+		FToolMenuEntry StopEntry = FToolMenuEntry::InitToolBarButton(
+			Commands.StopSoundWave,
+			LOCTEXT("WaveformEditorStopButton", ""),
+			LOCTEXT("WaveformEditorStopButtonTooltip", "Stops this SoundWave"),
+			FSlateIcon(FAppStyle::GetAppStyleSetName(), "PlayWorld.StopPlaySession.Small")
+		);
+
+		StopEntry.StyleNameOverride = FName("Toolbar.BackplateRight");
+
+		PlayBackSection.AddEntry(PlayEntry);
+		PlayBackSection.AddEntry(PauseEntry);
+		PlayBackSection.AddEntry(StopEntry);
+
+
+		FToolMenuInsert InsertAfterPlaybackSection("Transport Controls", EToolMenuInsertType::After);
+		FToolMenuSection& ZoomSection = ToolBar->AddSection("Zoom Controls", TAttribute<FText>(), InsertAfterPlaybackSection);
+
+		FToolMenuEntry ZoomInEntry = FToolMenuEntry::InitToolBarButton(
+			Commands.ZoomIn,
+			LOCTEXT("WaveformEditorZoomIn", ""),
+			LOCTEXT("WaveformEditorZoomInButtonTooltip", "Zooms into the soundwave"),
+			FSlateIcon(FAppStyle::GetAppStyleSetName(), "Icons.Plus")
+		);
+
+		FToolMenuEntry ZoomOutEntry = FToolMenuEntry::InitToolBarButton(
+			Commands.ZoomOut,
+			LOCTEXT("WaveformEditorZoomOut", ""),
+			LOCTEXT("WaveformEditorZoomOutButtonTooltip", "Zooms out the soundwave"),
+			FSlateIcon(FAppStyle::GetAppStyleSetName(), "Icons.Minus")
+		);
+
+		ZoomSection.AddEntry(ZoomInEntry);
+		ZoomSection.AddEntry(ZoomOutEntry);
+
+		FToolMenuInsert InsertAfterZoomSection("Zoom Controls", EToolMenuInsertType::After);
+		FToolMenuSection& ExportSection = ToolBar->AddSection("Export Controls", TAttribute<FText>(), InsertAfterZoomSection);
+
+		FToolMenuEntry ExportEntry = FToolMenuEntry::InitToolBarButton(
+			Commands.ExportWaveform,
+			LOCTEXT("WaveformEditorRender", ""),
+			LOCTEXT("WaveformEditorRenderButtonTooltip", "Exports the edited waveform to a USoundWave asset"),
+			FSlateIcon(FAppStyle::GetAppStyleSetName(), "GenericCurveEditor.SetViewModeAbsolute")
+		);
+
+		ExportSection.AddEntry(ExportEntry);
+
+	}
+
+	return true;
+}
+
+bool FWaveformEditor::BindCommands()
+{
+	const FWaveformEditorCommands& Commands = FWaveformEditorCommands::Get();
+
+	ToolkitCommands->MapAction(
+		Commands.PlaySoundWave,
+		FExecuteAction::CreateSP(TransportController.ToSharedRef(), &FWaveformEditorTransportController::Play),
+		FCanExecuteAction::CreateSP(this, &FWaveformEditor::CanPressPlayButton));
+
+	ToolkitCommands->MapAction(
+		Commands.StopSoundWave,
+		FExecuteAction::CreateSP(TransportController.ToSharedRef(), &FWaveformEditorTransportController::Stop),
+		FCanExecuteAction::CreateSP(TransportController.ToSharedRef(), &FWaveformEditorTransportController::CanStop));
+
+	ToolkitCommands->MapAction(
+		Commands.TogglePlayback,
+		FExecuteAction::CreateSP(TransportController.ToSharedRef(), &FWaveformEditorTransportController::TogglePlayback));
+
+	ToolkitCommands->MapAction(
+		Commands.PauseSoundWave,
+		FExecuteAction::CreateSP(TransportController.ToSharedRef(), &FWaveformEditorTransportController::Pause),
+		FCanExecuteAction::CreateSP(TransportController.ToSharedRef(), &FWaveformEditorTransportController::IsPlaying));
+
+	ToolkitCommands->MapAction(
+		Commands.ZoomIn,
+		FExecuteAction::CreateSP(ZoomManager.ToSharedRef(), &FWaveformEditorZoomController::ZoomIn),
+		FCanExecuteAction::CreateSP(ZoomManager.ToSharedRef(), &FWaveformEditorZoomController::CanZoomIn));
+
+	ToolkitCommands->MapAction(
+		Commands.ZoomOut,
+		FExecuteAction::CreateSP(ZoomManager.ToSharedRef(), &FWaveformEditorZoomController::ZoomOut),
+		FCanExecuteAction::CreateSP(ZoomManager.ToSharedRef(), &FWaveformEditorZoomController::CanZoomOut));
+
+	ToolkitCommands->MapAction(
+		Commands.ExportWaveform,
+		FExecuteAction::CreateSP(this, &FWaveformEditor::ExportWaveform),
+		FCanExecuteAction::CreateSP(WaveWriter.ToSharedRef(), &FWaveformEditorWaveWriter::CanCreateSoundWaveAsset));
+
+	return true;
+}
+
+
+FName FWaveformEditor::GetEditorName() const
+{
+	return EditorName;
+}
+
+FName FWaveformEditor::GetToolkitFName() const
+{
+	return ToolkitFName;
+}
+
+FText FWaveformEditor::GetBaseToolkitName() const
+{
+	return LOCTEXT("AppLabel", "Waveform Editor");
+}
+
+FString FWaveformEditor::GetWorldCentricTabPrefix() const
+{
+	return LOCTEXT("WorldCentricTabPrefix", "Waveform Editor").ToString();
+}
+
+FLinearColor FWaveformEditor::GetWorldCentricTabColorScale() const
+{
+	return FLinearColor(0.0f, 0.0f, 0.2f, 0.5f);
+}
+
+void FWaveformEditor::NotifyPostChange(const FPropertyChangedEvent& PropertyChangedEvent, class FEditPropertyChain* PropertyThatChanged)
+{	
+	TDoubleLinkedList<FProperty*>::TDoubleLinkedListNode* PropertyNode = PropertyThatChanged->GetActiveMemberNode();
+
+	if (!PropertyNode)
+	{
+		return;
+	}
+	
+	bool bIsTransformation = false;
+
+	do 
+	{
+		bIsTransformation |= PropertyNode->GetValue()->GetName() == TEXT("Transformations");
+		PropertyNode = PropertyNode->GetPrevNode();
+	} while (PropertyNode != nullptr);
+
+	if (!bIsTransformation)
+	{
+		return;
+	}
+
+	const bool bUpdateTransformationChain = PropertyChangedEvent.GetPropertyName() == TEXT("Transformations");
+
+	if (bUpdateTransformationChain)
+	{
+		TransformationsRenderManager->GenerateLayersChain();
+	}
+
+	TransformationsRenderManager->UpdateRenderElements();
+}
+
+bool FWaveformEditor::SetUpPropertiesView()
+{
+	if (SoundWave == nullptr)
+	{
+		UE_LOG(LogWaveformEditor, Warning, TEXT("Trying to setup wav editor properties view from a null SoundWave"));
+		return false;
+	}
+
+	FDetailsViewArgs Args;
+	Args.bHideSelectionTip = true;
+	Args.NotifyHook = this;
+
+	FPropertyEditorModule& PropertyModule = FModuleManager::LoadModuleChecked<FPropertyEditorModule>("PropertyEditor");
+	PropertiesView = PropertyModule.CreateDetailView(Args);
+	PropertiesView->SetObject(SoundWave);
+
+	return true;
+}
+
+TSharedRef<SDockTab> FWaveformEditor::SpawnTab_WaveformDisplay(const FSpawnTabArgs& Args)
+{
+	check(Args.GetTabId() == WaveformDisplayTabId);
+
+	return SNew(SDockTab)
+		.Label(LOCTEXT("WaveformDisplayTitle", "Waveform Display"))
+		[
+			WaveformPanel.ToSharedRef()
+		];
+}
+
+const TSharedRef<FTabManager::FLayout> FWaveformEditor::SetupStandaloneLayout()
+{
+	const TSharedRef<FTabManager::FLayout> StandaloneDefaultLayout = FTabManager::NewLayout("Standalone_WaveformEditor_v1")
+		->AddArea
+		(
+			FTabManager::NewPrimaryArea()->SetOrientation(Orient_Vertical)
+			->Split
+			(
+				FTabManager::NewSplitter()->SetOrientation(Orient_Horizontal)->SetSizeCoefficient(0.9f)
+				->Split
+				(
+					FTabManager::NewStack()
+					->SetSizeCoefficient(0.2f)
+					->AddTab(PropertiesTabId, ETabState::OpenedTab)
+					->SetForegroundTab(PropertiesTabId)
+
+				)
+				->Split
+				(
+					FTabManager::NewStack()
+					->SetSizeCoefficient(0.8f)
+					->SetHideTabWell(true)
+					->AddTab(WaveformDisplayTabId, ETabState::OpenedTab)
+				)
+			)
+		);
+
+	return StandaloneDefaultLayout;
+}
+
+
+
+TSharedRef<SDockTab> FWaveformEditor::SpawnTab_Properties(const FSpawnTabArgs& Args)
+{
+	check(Args.GetTabId() == PropertiesTabId);
+
+	return SNew(SDockTab)
+		.Label(LOCTEXT("SoundWaveDetailsTitle", "Details"))
+		[
+			PropertiesView.ToSharedRef()
+		];
+}
+
+bool FWaveformEditor::SetUpWaveformPanel()
+{
+	if (SoundWave == nullptr)
+	{
+		UE_LOG(LogWaveformEditor, Warning, TEXT("Trying to setup waveform panel from a null SoundWave"));
+		return false;
+	}
+
+	if (!SetUpZoom())
+	{
+		UE_LOG(LogWaveformEditor, Warning, TEXT("Failed to Init Zoom Manager while setting up waveform panel"));
+		return false;
+	}
+
+	TSharedPtr<FWaveformEditorRenderData> RenderData = MakeShared<FWaveformEditorRenderData>();
+	
+	TransportCoordinator = MakeShared<FWaveformEditorTransportCoordinator>(RenderData.ToSharedRef());
+	RenderData->OnRenderDataUpdated.AddSP(TransportCoordinator.Get(), &FWaveformEditorTransportCoordinator::HandleRenderDataUpdate);
+
+	TransformationsRenderManager = MakeShared<FWaveformTransformationsRenderManager>(SoundWave, RenderData.ToSharedRef(), TransportCoordinator.ToSharedRef(), ZoomManager.ToSharedRef());
+	TransformationsRenderManager->OnRenderDataGenerated.AddSP(RenderData.Get(), &FWaveformEditorRenderData::UpdateRenderData);
+
+	TSharedPtr<SWaveformTransformationsOverlay> TransformationsOverlay = SNew(SWaveformTransformationsOverlay, TransformationsRenderManager->GetTransformLayers());
+
+	TransformationsRenderManager->OnLayersChainGenerated.AddSP(TransformationsOverlay.Get(), &SWaveformTransformationsOverlay::OnLayerChainUpdate);
+	TransformationsRenderManager->UpdateRenderElements();
+
+	WaveformPanel = SNew(SWaveformPanel, RenderData.ToSharedRef(), TransportCoordinator.ToSharedRef(), ZoomManager.ToSharedRef(), TransformationsOverlay);
+
+	return WaveformPanel != nullptr;
+}
+
+void FWaveformEditor::HandlePlaybackPercentageChange(const UAudioComponent* InComponent, const USoundWave* InSoundWave, const float InPlaybackPercentage)
+{
+	const bool bIsStopped = AudioComponent->GetPlayState() == EAudioComponentPlayState::Stopped;
+	const bool bIsPaused = AudioComponent->GetPlayState() == EAudioComponentPlayState::Paused;
+	const bool bPropagatePercentage = !bIsStopped && !bIsPaused;
+	
+	if (InComponent == AudioComponent && bPropagatePercentage)
+	{
+		if (TransportCoordinator.IsValid())
+		{
+			const float ClampedPlayBackPercentage = FGenericPlatformMath::Fmod(InPlaybackPercentage, 1.f);
+			TransportCoordinator->ReceivePlayBackRatio(ClampedPlayBackPercentage);
+		}
+	}
+}
+
+void FWaveformEditor::HandleAudioComponentPlayStateChanged(const UAudioComponent* InAudioComponent, EAudioComponentPlayState NewPlayState)
+{
+	if (InAudioComponent != AudioComponent)
+	{
+		return;
+	}
+
+	switch (NewPlayState)
+	{
+	default:
+		break;
+	case EAudioComponentPlayState::Stopped:
+		if (!TransportCoordinator->IsScrubbing())
+		{
+			TransportController->CacheStartTime(0.f);
+			TransportCoordinator->Stop();
+		}
+		break;
+	}
+}
+
+void FWaveformEditor::HandlePlayheadScrub(const uint32 SelectedSample, const uint32 TotalSampleLength, const bool bIsMoving)
+{
+	if (bIsMoving)
+	{
+		if (TransportController->IsPlaying())
+		{
+			TransportController->Stop();
+			bWasPlayingBeforeScrubbing = true;
+		}
+	}
+	else
+	{
+		const float PlayBackRatio = SelectedSample / (float)TotalSampleLength;
+		const float NewTime = PlayBackRatio * SoundWave->Duration;
+
+		if (TransportController->IsPlaying())
+		{
+			TransportController->Seek(NewTime);
+			return;
+		}
+			
+		if (bWasPlayingBeforeScrubbing)
+		{
+			TransportController->Play(NewTime);
+			bWasPlayingBeforeScrubbing = false;
+		}
+		else
+		{
+			TransportController->CacheStartTime(NewTime);
+		}
+		
+	}
+}
+
+void FWaveformEditor::AddReferencedObjects(FReferenceCollector& Collector)
+{
+	Collector.AddReferencedObject(SoundWave);
+	Collector.AddReferencedObject(AudioComponent);
+}
+
+FString FWaveformEditor::GetReferencerName() const
+{
+	return TEXT("FWaveformEditor");
+}
+
+bool FWaveformEditor::CanPressPlayButton() const
+{
+	return TransportController->CanPlay() && (TransportController->IsPaused() || !TransportController->IsPlaying());
+}
+
+bool FWaveformEditor::SetUpWaveWriter()
+{
+	check(SoundWave)
+	WaveWriter = MakeShared<FWaveformEditorWaveWriter>(SoundWave);
+	return WaveWriter != nullptr;
+}
+
+void FWaveformEditor::ExportWaveform()
+{
+	check(WaveWriter);
+	WaveWriter->ExportTransformedWaveform();
+}
+
+#undef LOCTEXT_NAMESPACE
diff --git a/Engine/Plugins/Editor/WaveformEditor/Source/WaveformEditor/Private/WaveformEditorCommands.cpp b/Engine/Plugins/Editor/WaveformEditor/Source/WaveformEditor/Private/WaveformEditorCommands.cpp
new file mode 100644
--- /dev/null
+++ b/Engine/Plugins/Editor/WaveformEditor/Source/WaveformEditor/Private/WaveformEditorCommands.cpp
@@ -0,0 +1,33 @@
+// Copyright Epic Games, Inc. All Rights Reserved.
+
+#include "WaveformEditorCommands.h"
+
+#define LOCTEXT_NAMESPACE "WaveformEditorCommands"
+
+FWaveformEditorCommands::FWaveformEditorCommands() 
+	: TCommands<FWaveformEditorCommands>(
+		TEXT("WaveformEditor"), 
+		NSLOCTEXT("Contexts", "WaveformEditor", "Waveform Editor"), 
+		NAME_None, 
+		"AudioStyleSet"
+	)
+{}
+
+void FWaveformEditorCommands::RegisterCommands()
+{
+	UI_COMMAND(PlaySoundWave, "Play", "Play this sound wave", EUserInterfaceActionType::Button, FInputChord());
+	UI_COMMAND(PauseSoundWave, "Pause", "Pause this sound wave", EUserInterfaceActionType::Button, FInputChord());
+	UI_COMMAND(StopSoundWave, "Stop", "Stops playback and rewinds to beginning of file", EUserInterfaceActionType::Button, FInputChord(EKeys::Enter));
+	UI_COMMAND(TogglePlayback, "TogglePlayback", "Toggles between play and pause", EUserInterfaceActionType::Button, FInputChord(EKeys::SpaceBar));
+
+	UI_COMMAND(ZoomIn, "Zoom In", "Zooms into the sound wave", EUserInterfaceActionType::Button, FInputChord());
+	UI_COMMAND(ZoomOut, "Zoom Out", "Zooms out from the sound wave", EUserInterfaceActionType::Button, FInputChord());
+	UI_COMMAND(ExportWaveform, "Export Waveform", "Exports the edited waveform to a new sound wave uasset", EUserInterfaceActionType::Button, FInputChord());
+}
+	
+const FWaveformEditorCommands& FWaveformEditorCommands::Get()
+{
+	return TCommands<FWaveformEditorCommands>::Get();
+}
+
+#undef LOCTEXT_NAMESPACE
diff --git a/Engine/Plugins/Editor/WaveformEditor/Source/WaveformEditor/Private/WaveformEditorInstantiator.cpp b/Engine/Plugins/Editor/WaveformEditor/Source/WaveformEditor/Private/WaveformEditorInstantiator.cpp
new file mode 100644
--- /dev/null
+++ b/Engine/Plugins/Editor/WaveformEditor/Source/WaveformEditor/Private/WaveformEditorInstantiator.cpp
@@ -0,0 +1,101 @@
+// Copyright Epic Games, Inc. All Rights Reserved.
+
+#include "WaveformEditorInstantiator.h"
+
+#include "ContentBrowserModule.h"
+#include "Framework/MultiBox/MultiBoxBuilder.h"
+#include "Sound/SoundWave.h"
+#include "WaveformEditor.h"
+#include "WaveformEditorLog.h"
+
+#define LOCTEXT_NAMESPACE "WaveformEditorInstantiator"
+
+void FWaveformEditorInstantiator::ExtendContentBrowserSelectionMenu()
+{
+	FContentBrowserModule& ContentBrowserModule = FModuleManager::GetModuleChecked<FContentBrowserModule>(TEXT("ContentBrowser"));
+
+	TArray<FContentBrowserMenuExtender_SelectedAssets>& ContentBrowserExtenders = ContentBrowserModule.GetAllAssetViewContextMenuExtenders();
+	ContentBrowserExtenders.Add(FContentBrowserMenuExtender_SelectedAssets::CreateSP(this, &FWaveformEditorInstantiator::OnExtendContentBrowserAssetSelectionMenu));
+}
+
+TSharedRef<FExtender> FWaveformEditorInstantiator::OnExtendContentBrowserAssetSelectionMenu(const TArray<FAssetData>& SelectedAssets)
+{
+	TSharedRef<FExtender> Extender = MakeShared<FExtender>();
+
+	Extender->AddMenuExtension(
+		"GetAssetActions",
+		EExtensionHook::After,
+		nullptr,
+		FMenuExtensionDelegate::CreateSP(this, &FWaveformEditorInstantiator::AddWaveformEditorMenuEntry, SelectedAssets)
+	);
+
+	return Extender;
+}
+
+void FWaveformEditorInstantiator::AddWaveformEditorMenuEntry(FMenuBuilder& MenuBuilder, TArray<FAssetData> SelectedAssets)
+{
+	if (SelectedAssets.Num() > 0)
+	{
+		// check that all selected assets are USoundWaves.
+		TArray<USoundWave*> SelectedSoundWaves;
+		for (const FAssetData& SelectedAsset : SelectedAssets)
+		{
+			if (SelectedAsset.GetClass() != USoundWave::StaticClass())
+			{
+				return;
+			}
+
+			SelectedSoundWaves.Add(static_cast<USoundWave*>(SelectedAsset.GetAsset()));
+		}
+
+		MenuBuilder.AddMenuEntry(
+			LOCTEXT("SoundWave_WaveformEditor", "Edit Waveform"),
+			LOCTEXT("SoundWave_WaveformEditor_Tooltip", "Open waveform editor"),
+			FSlateIcon(),
+			FUIAction(
+				FExecuteAction::CreateSP(this, &FWaveformEditorInstantiator::CreateWaveformEditor, SelectedSoundWaves),
+				FCanExecuteAction()
+			)
+		);
+	}
+}
+
+bool FWaveformEditorInstantiator::CanSoundWaveBeOpenedInEditor(const USoundWave* SoundWaveToEdit)
+{
+	if (SoundWaveToEdit == nullptr)
+	{
+		UE_LOG(LogWaveformEditor, Warning, TEXT("Could not open waveform editor for soundwave %s, soundwave is null"), *(SoundWaveToEdit->GetName()))
+		return false;
+	}
+	else if (SoundWaveToEdit->GetDuration() == 0.f)
+	{
+		UE_LOG(LogWaveformEditor, Warning, TEXT("Could not open waveform editor for soundwave %s, duration is 0"), *(SoundWaveToEdit->GetName()))
+		return false;
+	}
+	else if (SoundWaveToEdit->NumChannels == 0.f)
+	{
+		UE_LOG(LogWaveformEditor, Warning, TEXT("Could not open waveform editor for soundwave %s, channel count is 0"), *(SoundWaveToEdit->GetName()))
+		return false;
+	}
+		
+	return true;
+}
+
+void FWaveformEditorInstantiator::CreateWaveformEditor(TArray<USoundWave*> SoundWavesToEdit)
+{
+	for (USoundWave* SoundWavePtr : SoundWavesToEdit)
+	{
+		if (CanSoundWaveBeOpenedInEditor(SoundWavePtr))
+		{
+			TSharedRef<FWaveformEditor> WaveformEditor = MakeShared<FWaveformEditor>();
+
+			if (!WaveformEditor->Init(EToolkitMode::Standalone, nullptr, SoundWavePtr))
+			{
+				UE_LOG(LogWaveformEditor, Warning, TEXT("Could not open waveform editor for soundwave %s, initialization failed"), *(SoundWavePtr->GetName()))
+				WaveformEditor->CloseWindow();
+			}
+		}
+	}
+}
+
+#undef LOCTEXT_NAMESPACE
\ No newline at end of file
diff --git a/Engine/Plugins/Editor/WaveformEditor/Source/WaveformEditor/Private/WaveformEditorInstantiator.h b/Engine/Plugins/Editor/WaveformEditor/Source/WaveformEditor/Private/WaveformEditorInstantiator.h
new file mode 100644
--- /dev/null
+++ b/Engine/Plugins/Editor/WaveformEditor/Source/WaveformEditor/Private/WaveformEditorInstantiator.h
@@ -0,0 +1,24 @@
+// Copyright Epic Games, Inc. All Rights Reserved.
+
+#pragma once
+
+#include "IWaveformEditorInstantiator.h"
+#include "Templates/SharedPointer.h"
+
+class FExtender;
+class FMenuBuilder;
+struct FAssetData;
+
+class FWaveformEditorInstantiator : public IWaveformEditorInstantiator,  public TSharedFromThis<FWaveformEditorInstantiator>
+{
+public:
+	virtual void ExtendContentBrowserSelectionMenu() override;
+	virtual ~FWaveformEditorInstantiator() = default;
+
+private:
+	virtual void CreateWaveformEditor(TArray<USoundWave*> SoundWavesToEdit) override;
+
+	TSharedRef<FExtender> OnExtendContentBrowserAssetSelectionMenu(const TArray<FAssetData>& SelectedAssets);
+	void AddWaveformEditorMenuEntry(FMenuBuilder& MenuBuilder, TArray<FAssetData> SelectedAssets);
+	bool CanSoundWaveBeOpenedInEditor(const USoundWave* SoundWaveToEdit);
+};
\ No newline at end of file
diff --git a/Engine/Plugins/Editor/WaveformEditor/Source/WaveformEditor/Private/WaveformEditorModule.cpp b/Engine/Plugins/Editor/WaveformEditor/Source/WaveformEditor/Private/WaveformEditorModule.cpp
new file mode 100644
--- /dev/null
+++ b/Engine/Plugins/Editor/WaveformEditor/Source/WaveformEditor/Private/WaveformEditorModule.cpp
@@ -0,0 +1,30 @@
+// Copyright Epic Games, Inc. All Rights Reserved.
+
+#include "WaveformEditorModule.h"
+
+#include "Modules/ModuleManager.h"
+#include "WaveformEditorCommands.h"
+#include "WaveformEditorInstantiator.h"
+#include "WaveformEditorLog.h"
+
+DEFINE_LOG_CATEGORY(LogWaveformEditor);
+
+void FWaveformEditorModule::StartupModule()
+{
+	FWaveformEditorCommands::Register();
+
+	WaveformEditorInstantiator = MakeShared<FWaveformEditorInstantiator>();
+	RegisterContentBrowserExtensions(WaveformEditorInstantiator.Get());
+}
+
+void FWaveformEditorModule::ShutdownModule()
+{
+	FWaveformEditorCommands::Unregister();
+}
+
+void FWaveformEditorModule::RegisterContentBrowserExtensions(IWaveformEditorInstantiator* Instantiator)
+{
+	WaveformEditorInstantiator->ExtendContentBrowserSelectionMenu();
+}
+
+IMPLEMENT_MODULE(FWaveformEditorModule, WaveformEditor);
diff --git a/Engine/Plugins/Editor/WaveformEditor/Source/WaveformEditor/Private/WaveformEditorTransportController.cpp b/Engine/Plugins/Editor/WaveformEditor/Source/WaveformEditor/Private/WaveformEditorTransportController.cpp
new file mode 100644
--- /dev/null
+++ b/Engine/Plugins/Editor/WaveformEditor/Source/WaveformEditor/Private/WaveformEditorTransportController.cpp
@@ -0,0 +1,114 @@
+// Copyright Epic Games, Inc. All Rights Reserved.
+
+#include "WaveformEditorTransportController.h"
+
+FWaveformEditorTransportController::FWaveformEditorTransportController(UAudioComponent* InAudioComponent)
+	: AudioComponent(InAudioComponent)
+{}
+
+FWaveformEditorTransportController::~FWaveformEditorTransportController()
+{
+	Stop();
+}
+
+void FWaveformEditorTransportController::Play()
+{
+	if (!CanPlay())
+	{
+		return;
+	}
+
+	if (IsPaused())
+	{
+		AudioComponent->SetPaused(false);
+		return;
+	}
+
+	AudioComponent->Play(CachedAudioStartTime);
+
+}
+
+void FWaveformEditorTransportController::Play(const float StartTime)
+{
+	if (!CanPlay())
+	{
+		return;
+	}
+
+	if (IsPaused())
+	{
+		AudioComponent->SetPaused(false);
+	}
+
+	CacheStartTime(StartTime);
+	AudioComponent->Play(CachedAudioStartTime);
+}
+
+void FWaveformEditorTransportController::Pause()
+{
+	AudioComponent->SetPaused(true);
+}
+
+
+void FWaveformEditorTransportController::Stop()
+{
+	if (!CanStop())
+	{
+		return;
+	}
+
+	Play(0.f);
+	AudioComponent->StopDelayed(0.1f);
+
+	if (IsPaused())
+	{
+		AudioComponent->SetPaused(false);
+	}
+}
+
+void FWaveformEditorTransportController::TogglePlayback()
+{
+	if (IsPlaying())
+	{
+		Pause();
+	}
+	else
+	{
+		Play();
+	}
+}
+
+bool FWaveformEditorTransportController::CanPlay() const
+{
+	return SoundBaseIsValid();
+}
+
+bool FWaveformEditorTransportController::CanStop() const
+{
+	return (IsPlaying() || IsPaused());
+}
+
+bool FWaveformEditorTransportController::IsPaused() const
+{
+	return AudioComponent->GetPlayState() == EAudioComponentPlayState::Paused;
+}
+
+bool FWaveformEditorTransportController::IsPlaying() const
+{
+	return AudioComponent->GetPlayState() == EAudioComponentPlayState::Playing;
+}
+
+void FWaveformEditorTransportController::CacheStartTime(const float StartTime)
+{
+	CachedAudioStartTime = StartTime;
+}
+
+void FWaveformEditorTransportController::Seek(const float SeekTime)
+{
+	AudioComponent->Play(SeekTime);
+}
+
+const bool FWaveformEditorTransportController::SoundBaseIsValid() const
+{
+	return AudioComponent->GetSound() != nullptr;
+}
diff --git a/Engine/Plugins/Editor/WaveformEditor/Source/WaveformEditor/Private/WaveformEditorWaveWriter.cpp b/Engine/Plugins/Editor/WaveformEditor/Source/WaveformEditor/Private/WaveformEditorWaveWriter.cpp
new file mode 100644
--- /dev/null
+++ b/Engine/Plugins/Editor/WaveformEditor/Source/WaveformEditor/Private/WaveformEditorWaveWriter.cpp
@@ -0,0 +1,110 @@
+// Copyright Epic Games, Inc. All Rights Reserved.
+
+#include "WaveformEditorWaveWriter.h"
+
+#include "AssetToolsModule.h"
+#include "DSP/FloatArrayMath.h"
+#include "FileHelpers.h"
+#include "Sound/SoundWave.h"
+#include "WaveformEditorLog.h"
+
+FWaveformEditorWaveWriter::FWaveformEditorWaveWriter(USoundWave* InSoundWave)
+	: SourceSoundWave(InSoundWave)
+	, WaveWriter(MakeUnique<Audio::FSoundWavePCMWriter>())
+{
+}
+
+bool FWaveformEditorWaveWriter::CanCreateSoundWaveAsset() const
+{
+	return SourceSoundWave != nullptr && WaveWriter != nullptr && WaveWriter->IsDone();
+}
+
+void FWaveformEditorWaveWriter::ExportTransformedWaveform()
+{
+	check(SourceSoundWave);
+	check(WaveWriter);
+
+	const FString DefaultSuffix = TEXT("_Edited");
+
+	FString AssetName;
+	FString PackageName;
+
+	FAssetToolsModule& AssetToolsModule = FModuleManager::Get().LoadModuleChecked<FAssetToolsModule>("AssetTools");
+	AssetToolsModule.Get().CreateUniqueAssetName(SourceSoundWave->GetOutermost()->GetName(), DefaultSuffix, PackageName, AssetName);
+	
+	FString AssetPath = FPaths::GetPath(PackageName);
+
+	//The wave writer will already be putting 'Game' in front of the provided asset path
+	AssetPath = AssetPath.Replace(TEXT("/Game"), TEXT(""), ESearchCase::CaseSensitive); 
+
+	Audio::TSampleBuffer<> BufferToWrite = GenerateSampleBuffer();
+
+	TFunction<void(const USoundWave*)> OnSoundWaveWritten = [AssetName, AssetPath](const USoundWave* ResultingWave) {
+		UE_LOG(LogWaveformEditor, Log, TEXT("Finished Exporting edited soundwave %s/%s"), *AssetPath, *AssetName);
+		TArray<UPackage*> PackagesToSave;
+
+		if(ResultingWave->GetPackage())
+		{
+			PackagesToSave.Add(ResultingWave->GetPackage());
+			FEditorFileUtils::PromptForCheckoutAndSave(PackagesToSave, false /*bCheckDirty*/, true /*bPromptToSave*/);
+		}
+	};
+
+	if (!WaveWriter->BeginWriteToSoundWave(AssetName, BufferToWrite, AssetPath, OnSoundWaveWritten))
+	{
+		UE_LOG(LogWaveformEditor, Log, TEXT("Exporting edited soundwave to %s/%s failed"), *AssetPath, *AssetName);
+	}
+}
+
+
+Audio::TSampleBuffer<> FWaveformEditorWaveWriter::GenerateSampleBuffer() const
+{
+	TArray<uint8> RawPCMData;
+	uint16 NumChannels;
+	uint32 SampleRate;
+
+	if (!SourceSoundWave->GetImportedSoundWaveData(RawPCMData, SampleRate, NumChannels))
+	{
+		UE_LOG(LogInit, Warning, TEXT("Failed to get imported soundwave data for file: %s. Edited waveform will not be rendered."), *SourceSoundWave->GetPathName());
+		return Audio::TSampleBuffer<>();
+	}
+
+	uint32 NumSamples = RawPCMData.Num() * sizeof(uint8) / sizeof(int16);
+
+	Audio::FAlignedFloatBuffer Buffer;
+	Buffer.SetNumUninitialized(NumSamples);
+
+	Audio::ArrayPcm16ToFloat(MakeArrayView((int16*)RawPCMData.GetData(), NumSamples), Buffer);
+
+	if (SourceSoundWave->Transformations.Num() > 0)
+	{
+		Audio::FWaveformTransformationWaveInfo TransformationInfo;
+
+		TransformationInfo.Audio = &Buffer;
+		TransformationInfo.NumChannels = NumChannels;
+		TransformationInfo.SampleRate = SampleRate;
+
+		TArray<Audio::FTransformationPtr> Transformations = SourceSoundWave->CreateTransformations();
+
+		for (const Audio::FTransformationPtr& Transformation : Transformations)
+		{
+			Transformation->ProcessAudio(TransformationInfo);
+		}
+
+		if (const float MaxValue = Audio::ArrayMaxAbsValue(Buffer) > 1.f)
+		{
+			Audio::ArrayMultiplyByConstantInPlace(Buffer, 1.f / MaxValue);
+		}
+
+		SampleRate = TransformationInfo.SampleRate;
+		NumChannels = TransformationInfo.NumChannels;
+		NumSamples = Buffer.Num();
+
+		check(NumChannels > 0);
+		check(SampleRate > 0);
+	}
+
+	Audio::TSampleBuffer<> OutBuffer(Buffer, NumChannels, SampleRate);
+
+	return MoveTemp(OutBuffer);
+}
\ No newline at end of file
diff --git a/Engine/Plugins/Editor/WaveformEditor/Source/WaveformEditor/Private/WaveformEditorWaveWriter.h b/Engine/Plugins/Editor/WaveformEditor/Source/WaveformEditor/Private/WaveformEditorWaveWriter.h
new file mode 100644
--- /dev/null
+++ b/Engine/Plugins/Editor/WaveformEditor/Source/WaveformEditor/Private/WaveformEditorWaveWriter.h
@@ -0,0 +1,22 @@
+// Copyright Epic Games, Inc. All Rights Reserved.
+
+#pragma once
+
+#include "Sound/SampleBufferIO.h"
+
+class USoundWave;
+
+class FWaveformEditorWaveWriter
+{
+public:
+	explicit FWaveformEditorWaveWriter(USoundWave* InSoundWave);
+
+	bool CanCreateSoundWaveAsset() const;
+	void ExportTransformedWaveform();
+
+private: 
+	Audio::TSampleBuffer<> GenerateSampleBuffer() const;
+
+	USoundWave* SourceSoundWave = nullptr;
+	TUniquePtr<Audio::FSoundWavePCMWriter> WaveWriter = nullptr;
+};
\ No newline at end of file
diff --git a/Engine/Plugins/Editor/WaveformEditor/Source/WaveformEditor/Public/IWaveformEditorInstantiator.h b/Engine/Plugins/Editor/WaveformEditor/Source/WaveformEditor/Public/IWaveformEditorInstantiator.h
new file mode 100644
--- /dev/null
+++ b/Engine/Plugins/Editor/WaveformEditor/Source/WaveformEditor/Public/IWaveformEditorInstantiator.h
@@ -0,0 +1,16 @@
+// Copyright Epic Games, Inc. All Rights Reserved.
+
+#pragma once
+
+#include "Containers/Array.h"
+
+
+class USoundWave;
+class IWaveformEditorInstantiator
+{
+public:
+	virtual void ExtendContentBrowserSelectionMenu() = 0;
+
+private:
+	virtual void CreateWaveformEditor(TArray<USoundWave*> SoundWavesToEdit) = 0;
+};
\ No newline at end of file
diff --git a/Engine/Plugins/Editor/WaveformEditor/Source/WaveformEditor/Public/IWaveformEditorModule.h b/Engine/Plugins/Editor/WaveformEditor/Source/WaveformEditor/Public/IWaveformEditorModule.h
new file mode 100644
--- /dev/null
+++ b/Engine/Plugins/Editor/WaveformEditor/Source/WaveformEditor/Public/IWaveformEditorModule.h
@@ -0,0 +1,13 @@
+// Copyright Epic Games, Inc. All Rights Reserved.
+
+#pragma once
+
+#include "Modules/ModuleInterface.h"
+
+class IWaveformEditorInstantiator;
+
+class IWaveformEditorModule : public IModuleInterface
+{
+	/** Registers waveform editor asset actions. */
+	virtual void RegisterContentBrowserExtensions(IWaveformEditorInstantiator* Instantiator) = 0;
+};
diff --git a/Engine/Plugins/Editor/WaveformEditor/Source/WaveformEditor/Public/WaveformEditor.h b/Engine/Plugins/Editor/WaveformEditor/Source/WaveformEditor/Public/WaveformEditor.h
new file mode 100644
--- /dev/null
+++ b/Engine/Plugins/Editor/WaveformEditor/Source/WaveformEditor/Public/WaveformEditor.h
@@ -0,0 +1,104 @@
+// Copyright Epic Games, Inc. All Rights Reserved.
+
+#pragma once
+
+#include "Misc/NotifyHook.h"
+#include "Toolkits/AssetEditorToolkit.h"
+#include "WaveformEditorTransportController.h"
+#include "WaveformEditorTransportCoordinator.h"
+#include "WaveformEditorZoomController.h"
+
+class IToolkitHost;
+class SDockTab;
+class USoundWave;
+class UAudioComponent;
+
+class WAVEFORMEDITOR_API FWaveformEditor : public FAssetEditorToolkit, public FGCObject, public FNotifyHook
+{
+public:
+	bool Init(const EToolkitMode::Type Mode, const TSharedPtr<IToolkitHost>& InitToolkitHost, USoundWave* SoundWaveToEdit);
+
+	/** FAssetEditorToolkit interface */
+	virtual void RegisterTabSpawners(const TSharedRef<FTabManager>& TabManager) override;
+	virtual void UnregisterTabSpawners(const TSharedRef<FTabManager>& TabManager) override;
+	virtual FName GetEditorName() const override;
+	virtual FName GetToolkitFName() const override;
+	virtual FText GetBaseToolkitName() const override;
+	virtual FString GetWorldCentricTabPrefix() const override;
+	virtual FLinearColor GetWorldCentricTabColorScale() const override;
+
+	void NotifyPreChange(class FEditPropertyChain* PropertyAboutToChange) override {};
+	void NotifyPostChange(const FPropertyChangedEvent& PropertyChangedEvent, class FEditPropertyChain* PropertyThatChanged) override;
+
+private:
+	bool SetupAudioComponent();
+	bool SetUpTransportController();
+	bool SetUpZoom();
+
+	bool BindDelegates();
+
+	/**	Sets the wave editor layout */
+	const TSharedRef<FTabManager::FLayout> SetupStandaloneLayout();
+
+	/**	Toolbar Setup */
+	bool RegisterToolbar();
+	bool BindCommands();
+
+	/**	Properties tab set up */
+	TSharedRef<SDockTab> SpawnTab_Properties(const FSpawnTabArgs& Args);
+	bool SetUpPropertiesView();
+
+	/**	Waveform view tab setup */
+	TSharedRef<SDockTab> SpawnTab_WaveformDisplay(const FSpawnTabArgs& Args);
+	bool SetUpWaveformPanel();
+
+	/** Playback delegates handlers */
+	void HandlePlaybackPercentageChange(const UAudioComponent* InComponent, const USoundWave* InSoundWave, const float InPlaybackPercentage);
+	void HandleAudioComponentPlayStateChanged(const UAudioComponent* InAudioComponent, EAudioComponentPlayState NewPlayState);
+	void HandlePlayheadScrub(const uint32 SelectedSample, const uint32 TotalSampleLength, const bool bIsMoving);
+
+	/** FGCObject overrides */
+	virtual void AddReferencedObjects(FReferenceCollector& Collector) override;
+	virtual FString GetReferencerName() const override;
+
+	bool CanPressPlayButton() const;
+
+	bool SetUpWaveWriter();
+	void ExportWaveform();
+
+	/** Waveform Preview widget */
+	TSharedPtr<class SWaveformPanel> WaveformPanel;
+
+	/** Manages render information for waveform transforms */
+	TSharedPtr<class FWaveformTransformationsRenderManager> TransformationsRenderManager = nullptr;
+
+	/** Exports the edited waveform to a new asset */
+	TSharedPtr<class FWaveformEditorWaveWriter> WaveWriter = nullptr;
+
+	/** Manages Transport info in waveform panel */
+	TSharedPtr<FWaveformEditorTransportCoordinator> TransportCoordinator = nullptr;
+
+	/** Controls Transport of the audio component  */
+	TSharedPtr<FWaveformEditorTransportController> TransportController = nullptr;
+
+	/** Controls and propagates zoom level */
+	TSharedPtr<FWaveformEditorZoomController> ZoomManager = nullptr;
+
+	/** Properties tab */
+	TSharedPtr<IDetailsView> PropertiesView;
+
+	/** Settings Editor App Identifier */
+	static const FName AppIdentifier;
+
+	/** Tab Ids */
+	static const FName PropertiesTabId;
+	static const FName WaveformDisplayTabId;
+
+	static const FName EditorName;
+	static const FName ToolkitFName;
+
+	USoundWave* SoundWave = nullptr;
+	UAudioComponent* AudioComponent = nullptr;
+
+	bool bWasPlayingBeforeScrubbing = false;
+};
diff --git a/Engine/Plugins/Editor/WaveformEditor/Source/WaveformEditor/Public/WaveformEditorCommands.h b/Engine/Plugins/Editor/WaveformEditor/Source/WaveformEditor/Public/WaveformEditorCommands.h
new file mode 100644
--- /dev/null
+++ b/Engine/Plugins/Editor/WaveformEditor/Source/WaveformEditor/Public/WaveformEditorCommands.h
@@ -0,0 +1,27 @@
+// Copyright Epic Games, Inc. All Rights Reserved.
+
+#pragma once
+
+#include "CoreMinimal.h"
+
+#include "Framework/Commands/Commands.h"
+
+class FWaveformEditorCommands : public TCommands<FWaveformEditorCommands>
+{
+public:
+	FWaveformEditorCommands();
+
+	virtual void RegisterCommands() override;
+	static const FWaveformEditorCommands& Get();
+
+	// Soundwave Playback controls
+	TSharedPtr<FUICommandInfo> PlaySoundWave;
+	TSharedPtr<FUICommandInfo> PauseSoundWave;
+	TSharedPtr<FUICommandInfo> StopSoundWave;
+	TSharedPtr<FUICommandInfo> TogglePlayback;
+
+	TSharedPtr<FUICommandInfo> ZoomIn;
+	TSharedPtr<FUICommandInfo> ZoomOut;
+
+	TSharedPtr<FUICommandInfo> ExportWaveform;
+};
diff --git a/Engine/Plugins/Editor/WaveformEditor/Source/WaveformEditor/Public/WaveformEditorLog.h b/Engine/Plugins/Editor/WaveformEditor/Source/WaveformEditor/Public/WaveformEditorLog.h
new file mode 100644
--- /dev/null
+++ b/Engine/Plugins/Editor/WaveformEditor/Source/WaveformEditor/Public/WaveformEditorLog.h
@@ -0,0 +1,7 @@
+// Copyright Epic Games, Inc. All Rights Reserved.
+
+#pragma once
+
+#include "CoreMinimal.h"
+
+DECLARE_LOG_CATEGORY_EXTERN(LogWaveformEditor, Log, All);
\ No newline at end of file
diff --git a/Engine/Plugins/Editor/WaveformEditor/Source/WaveformEditor/Public/WaveformEditorModule.h b/Engine/Plugins/Editor/WaveformEditor/Source/WaveformEditor/Public/WaveformEditorModule.h
new file mode 100644
--- /dev/null
+++ b/Engine/Plugins/Editor/WaveformEditor/Source/WaveformEditor/Public/WaveformEditorModule.h
@@ -0,0 +1,21 @@
+// Copyright Epic Games, Inc. All Rights Reserved.
+
+#pragma once
+
+#include "IWaveformEditorModule.h"
+#include "Templates/SharedPointer.h"
+
+class IWaveformEditorInstantiator;
+
+class FWaveformEditorModule : public IWaveformEditorModule
+{
+public:
+	virtual void StartupModule() override;
+	virtual void ShutdownModule() override;
+
+private:
+	virtual void RegisterContentBrowserExtensions(IWaveformEditorInstantiator* Instantiator) override;
+
+	TSharedPtr<IWaveformEditorInstantiator> WaveformEditorInstantiator = nullptr;
+};
+
diff --git a/Engine/Plugins/Editor/WaveformEditor/Source/WaveformEditor/Public/WaveformEditorTransportController.h b/Engine/Plugins/Editor/WaveformEditor/Source/WaveformEditor/Public/WaveformEditorTransportController.h
new file mode 100644
--- /dev/null
+++ b/Engine/Plugins/Editor/WaveformEditor/Source/WaveformEditor/Public/WaveformEditorTransportController.h
@@ -0,0 +1,32 @@
+// Copyright Epic Games, Inc. All Rights Reserved.
+
+#pragma once
+
+#include "Components/AudioComponent.h"
+
+class FWaveformEditorTransportController
+{
+public:
+	explicit FWaveformEditorTransportController(UAudioComponent* InAudioComponent);
+	~FWaveformEditorTransportController();
+
+	void Play();
+	void Play(const float StartTime);
+	void Pause();
+	void Stop();
+	void TogglePlayback();
+	bool CanPlay() const;
+	bool CanStop() const;
+	bool IsPaused() const;
+	bool IsPlaying() const;
+	void CacheStartTime(const float StartTime);
+	void Seek(const float SeekTime);
+
+private:
+	const bool SoundBaseIsValid() const;
+
+	UAudioComponent* AudioComponent = nullptr;
+	float CachedAudioStartTime = 0.f;
+
+	FDelegateHandle PlayStateChangeDelegateHandle;
+};
\ No newline at end of file
diff --git a/Engine/Plugins/Editor/WaveformEditor/Source/WaveformEditor/WaveformEditor.build.cs b/Engine/Plugins/Editor/WaveformEditor/Source/WaveformEditor/WaveformEditor.build.cs
new file mode 100644
--- /dev/null
+++ b/Engine/Plugins/Editor/WaveformEditor/Source/WaveformEditor/WaveformEditor.build.cs
@@ -0,0 +1,34 @@
+// Copyright Epic Games, Inc. All Rights Reserved.
+
+using UnrealBuildTool;
+using System.IO;
+
+namespace UnrealBuildTool.Rules
+{
+	public class WaveformEditor : ModuleRules
+	{
+		public WaveformEditor(ReadOnlyTargetRules Target) : base(Target)
+		{
+			PCHUsage = PCHUsageMode.UseExplicitOrSharedPCHs;
+
+			PublicDependencyModuleNames.AddRange(
+				new string[] {
+					"AudioExtensions",
+					"AssetRegistry",
+					"Core", 
+					"CoreUObject",
+					"EditorStyle",
+					"Engine",
+					"InputCore",
+					"SignalProcessing",
+					"UnrealEd", 
+					"Slate",
+					"SlateCore",
+					"ToolMenus",
+					"WaveformEditorWidgets",
+					"WaveformTransformationsWidgets"
+				}
+			);
+		}
+	}
+}
\ No newline at end of file
diff --git a/Engine/Plugins/Editor/WaveformEditor/Source/WaveformEditorWidgets/Private/SWaveformEditorTimeRuler.cpp b/Engine/Plugins/Editor/WaveformEditor/Source/WaveformEditorWidgets/Private/SWaveformEditorTimeRuler.cpp
new file mode 100644
--- /dev/null
+++ b/Engine/Plugins/Editor/WaveformEditor/Source/WaveformEditorWidgets/Private/SWaveformEditorTimeRuler.cpp
@@ -0,0 +1,256 @@
+// Copyright Epic Games, Inc. All Rights Reserved.
+
+#include "SWaveformEditorTimeRuler.h"
+
+#include "Layout/WidgetPath.h"
+#include "Framework/Application/SlateApplication.h"
+#include "Framework/MultiBox/MultiBoxBuilder.h"
+#include "WaveformEditorRenderData.h"
+#include "WaveformEditorTransportCoordinator.h"
+
+
+#define LOCTEXT_NAMESPACE "WaveformEditorTimeRuler"
+
+void SWaveformEditorTimeRuler::Construct(const FArguments& InArgs, TSharedRef<FWaveformEditorTransportCoordinator> InTransportCoordinator, TSharedRef<FWaveformEditorRenderData> InRenderData)
+{
+	TransportCoordinator = InTransportCoordinator;
+	RenderData = InRenderData;
+
+	DisplayUnit = InArgs._DisplayUnit;
+
+	Style = InArgs._Style;
+	
+	check(Style)
+	HandleColor = &Style->HandleColor;
+	HandleWidth = Style->HandleWidth;
+	HandleBrush = &Style->HandleBrush;
+	BackgroundColor = &Style->BackgroundColor;
+	BackgroundBrush = &Style->BackgroundBrush;
+	TicksColor = &Style->TicksColor;
+	TicksTextColor = &Style->TicksTextColor;
+	TicksTextFont = &Style->TicksTextFont;
+	DesiredWidth = Style->DesiredWidth;
+	DesiredHeight = Style->DesiredHeight;
+	TicksTextOffset = Style->TicksTextOffset;
+}
+
+FVector2D SWaveformEditorTimeRuler::ComputeDesiredSize(float) const
+{
+	return FVector2D(DesiredWidth, DesiredHeight);
+}
+
+int32 SWaveformEditorTimeRuler::OnPaint(const FPaintArgs& Args, const FGeometry& AllottedGeometry, const FSlateRect& MyCullingRect, FSlateWindowElementList& OutDrawElements, int32 LayerId, const FWidgetStyle& InWidgetStyle, bool bParentEnabled) const
+{
+	FSlateDrawElement::MakeBox(
+		OutDrawElements,
+		++LayerId,
+		AllottedGeometry.ToPaintGeometry(),
+		BackgroundBrush,
+		ESlateDrawEffect::None,
+		BackgroundColor->GetSpecifiedColor());
+	
+	DrawRulerTicks(AllottedGeometry, OutDrawElements, LayerId);
+	DrawPlayheadHandle(AllottedGeometry, OutDrawElements, LayerId);
+
+	return LayerId;
+}
+
+void SWaveformEditorTimeRuler::DrawRulerTicks(const FGeometry& AllottedGeometry, FSlateWindowElementList& OutDrawElements, int32& LayerId) const
+{
+	const double MinorGridXStep = GridMetrics.MajorGridXStep / GridMetrics.NumMinorGridDivisions;
+	const double MajorTickY = AllottedGeometry.Size.Y * 0.25;
+	const double MinorTickY = AllottedGeometry.Size.Y * 0.75;
+
+	TArray<FVector2D> LinePoints;
+	LinePoints.SetNumUninitialized(2);
+
+	for (double CurrentMajorTickX = GridMetrics.FirstMajorTickX; CurrentMajorTickX < AllottedGeometry.Size.X; CurrentMajorTickX += GridMetrics.MajorGridXStep)
+	{
+		const double MajorTickX = CurrentMajorTickX;
+		
+		LinePoints[0] = FVector2D(MajorTickX, MajorTickY);
+		LinePoints[1] = FVector2D(MajorTickX, AllottedGeometry.Size.Y);
+
+		FSlateDrawElement::MakeLines(
+			OutDrawElements,
+			++LayerId,
+			AllottedGeometry.ToPaintGeometry(),
+			LinePoints,
+			ESlateDrawEffect::None,
+			TicksColor->GetSpecifiedColor(),
+			false);
+
+
+		for (int32 MinorTickIndex = 1; MinorTickIndex < GridMetrics.NumMinorGridDivisions; ++MinorTickIndex)
+		{
+			const double MinorTickX = MajorTickX + MinorGridXStep * MinorTickIndex;
+			
+			LinePoints[0] = FVector2D(MinorTickX, MinorTickY);
+			LinePoints[1] = FVector2D(MinorTickX, AllottedGeometry.Size.Y);
+
+			FSlateDrawElement::MakeLines(
+				OutDrawElements,
+				++LayerId,
+				AllottedGeometry.ToPaintGeometry(),
+				LinePoints,
+				ESlateDrawEffect::None,
+				TicksColor->GetSpecifiedColor(),
+				false);
+
+		}
+
+		float TickTimeSeconds = MajorTickX / GridMetrics.PixelsPerSecond + GridMetrics.StartTime;
+		DrawTickTimeString(TickTimeSeconds, MajorTickX, MajorTickY, OutDrawElements, LayerId, AllottedGeometry);
+	}
+}
+
+void SWaveformEditorTimeRuler::DrawTickTimeString(float TickTimeSeconds, const double TickX, const double TickY, FSlateWindowElementList& OutDrawElements, int32& LayerId, const FGeometry& AllottedGeometry) const
+{
+	FString TimeString;
+
+	switch (DisplayUnit)
+	{
+	case EWaveformEditorDisplayUnit::Samples:
+	{
+		uint32 SampleTime = TickTimeSeconds * RenderData->GetSampleRate();
+		TimeString = FString::Printf(TEXT("%d"), SampleTime);
+	}
+	break;
+	case EWaveformEditorDisplayUnit::Seconds:
+		TimeString = FString::Printf(TEXT("%.3f"), TickTimeSeconds);
+		break;
+	default:
+		TimeString = FString::Printf(TEXT("%.3f"), TickTimeSeconds);
+		break;
+	}
+
+	FVector2D TextOffset(TickX + TicksTextOffset, TickY);
+
+	FSlateDrawElement::MakeText(
+		OutDrawElements,
+		++LayerId,
+		AllottedGeometry.ToPaintGeometry(TextOffset, AllottedGeometry.Size),
+		TimeString,
+		*TicksTextFont,
+		ESlateDrawEffect::None,
+		TicksTextColor->GetSpecifiedColor()
+	);
+}
+
+TSharedRef<SWidget> SWaveformEditorTimeRuler::MakeContextMenu()
+{
+	const bool bCloseAfterSelection = false;
+	const bool bCloseSelfOnly = false;
+	FMenuBuilder MenuBuilder(bCloseAfterSelection, NULL, TSharedPtr<FExtender>(), bCloseSelfOnly, &FCoreStyle::Get());
+	MenuBuilder.BeginSection("Ruler Options", LOCTEXT("RulerOptionsHeading", "Options"));
+	{
+		MenuBuilder.AddSubMenu(
+			LOCTEXT("TimeUnitsSubMenuLabel", "Time Unit"),
+			LOCTEXT("TimeUnitsSubMenuLabelToolTip", "Select the time unit shown by the grid"),
+			FNewMenuDelegate::CreateSP(this, &SWaveformEditorTimeRuler::MakeTimeUnitsSubMenu)
+		);
+	}
+	MenuBuilder.EndSection();
+
+	
+	return MenuBuilder.MakeWidget();
+}
+
+void SWaveformEditorTimeRuler::MakeTimeUnitsSubMenu(FMenuBuilder& SubMenuBuilder)
+{
+	SubMenuBuilder.BeginSection(NAME_None);
+
+	SubMenuBuilder.AddMenuEntry(
+		LOCTEXT("TimeUnitsSubMenuEntry_Seconds", "Seconds"),
+		LOCTEXT("TimeUnitsSubMenuEntry_Seconds", "Seconds"),
+		FSlateIcon(),
+		FUIAction(
+			FExecuteAction::CreateSP(this, &SWaveformEditorTimeRuler::NotifyTimeUnitMenuSelection, EWaveformEditorDisplayUnit::Seconds),
+			FCanExecuteAction::CreateLambda([this] { return DisplayUnit != EWaveformEditorDisplayUnit::Seconds; }),
+			FIsActionChecked::CreateLambda([this] { return DisplayUnit == EWaveformEditorDisplayUnit::Seconds;})
+		),
+		NAME_None,
+		EUserInterfaceActionType::ToggleButton
+	);
+
+	SubMenuBuilder.AddMenuEntry(
+		LOCTEXT("TimeUnitsSubMenuEntry_Samples", "Samples"),
+		LOCTEXT("TimeUnitsSubMenuEntry_Samples", "Samples"),
+		FSlateIcon(),
+		FUIAction(
+			FExecuteAction::CreateSP(this, &SWaveformEditorTimeRuler::NotifyTimeUnitMenuSelection, EWaveformEditorDisplayUnit::Samples),
+			FCanExecuteAction::CreateLambda([this] { return DisplayUnit != EWaveformEditorDisplayUnit::Samples; }),
+			FIsActionChecked::CreateLambda([this] { return DisplayUnit == EWaveformEditorDisplayUnit::Samples; })
+		),
+		NAME_None,
+		EUserInterfaceActionType::ToggleButton
+	);
+
+	SubMenuBuilder.EndSection();
+ }
+
+void SWaveformEditorTimeRuler::NotifyTimeUnitMenuSelection(const EWaveformEditorDisplayUnit SelectedDisplayUnit) const
+{
+	if (SelectedDisplayUnit == DisplayUnit)
+	{
+		return;
+	}
+
+	if (OnTimeUnitMenuSelection.IsBound())
+	{
+		OnTimeUnitMenuSelection.Broadcast(SelectedDisplayUnit);
+	}
+}
+
+void SWaveformEditorTimeRuler::UpdateGridMetrics(const FWaveEditorGridMetrics& InMetrics)
+{
+	GridMetrics = InMetrics;
+}
+
+void SWaveformEditorTimeRuler::UpdateDisplayUnit(const EWaveformEditorDisplayUnit InDisplayUnit)
+{
+	DisplayUnit = InDisplayUnit;
+}
+
+void SWaveformEditorTimeRuler::DrawPlayheadHandle(const FGeometry& AllottedGeometry, FSlateWindowElementList& OutDrawElements, int32& LayerId) const
+{
+	const float WindowWidth = AllottedGeometry.Size.X;
+	const float	HandleStart = TransportCoordinator->GetPlayheadPosition() * WindowWidth - HandleWidth / 2;
+	const float	HandleEnd = HandleStart + HandleWidth;
+ 	FPaintGeometry HandleGeometry = AllottedGeometry.ToPaintGeometry(FVector2D(HandleStart, 0), FVector2D(HandleEnd - HandleStart, AllottedGeometry.Size.Y));
+
+	FSlateDrawElement::MakeBox(
+		OutDrawElements,
+		LayerId,
+		HandleGeometry,
+		HandleBrush,
+		ESlateDrawEffect::None,
+		HandleColor->GetSpecifiedColor()
+	);
+
+	++LayerId;
+}
+
+FReply SWaveformEditorTimeRuler::OnMouseButtonDown(const FGeometry& MyGeometry, const FPointerEvent& MouseEvent)
+{
+	return TransportCoordinator->ReceiveMouseButtonDown(*this, MyGeometry, MouseEvent);
+}
+
+FReply SWaveformEditorTimeRuler::OnMouseButtonUp(const FGeometry& MyGeometry, const FPointerEvent& MouseEvent)
+{
+	const bool HandleRightMouseButton = MouseEvent.GetEffectingButton() == EKeys::RightMouseButton && HasMouseCapture();
+
+	if (HandleRightMouseButton)
+	{
+		FWidgetPath WidgetPath = MouseEvent.GetEventPath() != nullptr ? *MouseEvent.GetEventPath() : FWidgetPath();
+		FSlateApplication::Get().PushMenu(AsShared(), WidgetPath, MakeContextMenu(), FSlateApplication::Get().GetCursorPos(), FPopupTransitionEffect::ContextMenu);
+		return FReply::Handled();
+	}
+
+	return TransportCoordinator->ReceiveMouseButtonUp(*this, MyGeometry, MouseEvent);
+}
+
+FReply SWaveformEditorTimeRuler::OnMouseMove(const FGeometry& MyGeometry, const FPointerEvent& MouseEvent)
+{
+	return TransportCoordinator->ReceiveMouseMove(*this, MyGeometry, MouseEvent);
+}
diff --git a/Engine/Plugins/Editor/WaveformEditor/Source/WaveformEditorWidgets/Private/SWaveformEditorTimeRuler.h b/Engine/Plugins/Editor/WaveformEditor/Source/WaveformEditorWidgets/Private/SWaveformEditorTimeRuler.h
new file mode 100644
--- /dev/null
+++ b/Engine/Plugins/Editor/WaveformEditor/Source/WaveformEditorWidgets/Private/SWaveformEditorTimeRuler.h
@@ -0,0 +1,80 @@
+// Copyright Epic Games, Inc. All Rights Reserved.
+
+#pragma once
+
+#include "Framework/MultiBox/MultiBoxBuilder.h"
+#include "Styling/SlateStyleRegistry.h"
+#include "Styling/SlateWidgetStyleAsset.h"
+#include "WaveformEditorDisplayUnit.h"
+#include "WaveformEditorGridData.h"
+#include "WaveformEditorSlateTypes.h"
+#include "Widgets/SCompoundWidget.h"
+
+DECLARE_MULTICAST_DELEGATE_OneParam(FOnTimeUnitMenuSelection, const EWaveformEditorDisplayUnit /* Requested Display Unit */);
+
+class FWaveformEditorRenderData;
+class FWaveformEditorTransportCoordinator;
+
+class SWaveformEditorTimeRuler : public SCompoundWidget
+{
+
+	SLATE_BEGIN_ARGS(SWaveformEditorTimeRuler)
+		: _DisplayUnit(EWaveformEditorDisplayUnit::Seconds)
+	{
+		const ISlateStyle* WaveEditorStyle = FSlateStyleRegistry::FindSlateStyle("WaveformEditorStyle");
+		check(WaveEditorStyle)
+		_Style = &WaveEditorStyle->GetWidgetStyle<FWaveformEditorTimeRulerStyle>("WaveformEditorRuler.Style");
+	}
+
+	SLATE_ARGUMENT(EWaveformEditorDisplayUnit, DisplayUnit)
+
+	SLATE_STYLE_ARGUMENT(FWaveformEditorTimeRulerStyle, Style)
+	
+	SLATE_END_ARGS()
+
+public:
+	void Construct(const FArguments& InArgs, TSharedRef<FWaveformEditorTransportCoordinator> InTransportCoordinator, TSharedRef<FWaveformEditorRenderData> InRenderData);
+	void UpdateGridMetrics(const FWaveEditorGridMetrics& InMetrics);
+	void UpdateDisplayUnit(const EWaveformEditorDisplayUnit InDisplayUnit);
+
+	/** Delegate sent when the user selects a new display unit from the RMB menu*/
+	FOnTimeUnitMenuSelection OnTimeUnitMenuSelection;
+
+private:
+	virtual int32 OnPaint(const FPaintArgs& Args, const FGeometry& AllottedGeometry, const FSlateRect& MyCullingRect, FSlateWindowElementList& OutDrawElements, int32 LayerId, const FWidgetStyle& InWidgetStyle, bool bParentEnabled) const override;
+	virtual FVector2D ComputeDesiredSize(float) const override;
+
+	virtual FReply OnMouseButtonDown(const FGeometry& MyGeometry, const FPointerEvent& MouseEvent) override;
+	virtual FReply OnMouseButtonUp(const FGeometry& MyGeometry, const FPointerEvent& MouseEvent) override;
+	virtual FReply OnMouseMove(const FGeometry& MyGeometry, const FPointerEvent& MouseEvent) override;
+
+	void DrawPlayheadHandle(const FGeometry& AllottedGeometry, FSlateWindowElementList& OutDrawElements, int32& LayerId) const;
+	void DrawRulerTicks(const FGeometry& AllottedGeometry, FSlateWindowElementList& OutDrawElements, int32& LayerId) const;
+	void DrawTickTimeString(float TickTimeSeconds, const double TickX, const double TickY, FSlateWindowElementList& OutDrawElements, int32& LayerId, const FGeometry& AllottedGeometry) const;
+
+	TSharedRef<SWidget> MakeContextMenu();
+	void MakeTimeUnitsSubMenu(FMenuBuilder& SubMenuBuilder);
+	void NotifyTimeUnitMenuSelection(const EWaveformEditorDisplayUnit SelectedDisplayUnit) const;
+
+	FWaveEditorGridMetrics GridMetrics;
+
+	const FWaveformEditorTimeRulerStyle* Style = nullptr;
+
+	const FSlateColor* BackgroundColor = nullptr;
+	const FSlateBrush* BackgroundBrush = nullptr;
+	const FSlateBrush* HandleBrush = nullptr;
+	const FSlateColor* HandleColor = nullptr;
+	const FSlateColor* TicksColor = nullptr;
+	const FSlateColor* TicksTextColor = nullptr;
+	float DesiredHeight = 0.f;
+	float DesiredWidth = 0.f;
+	float HandleWidth = 0.f;
+	float TicksTextOffset = 5.f;
+	const FSlateFontInfo* TicksTextFont;
+
+	TSharedPtr<FWaveformEditorTransportCoordinator> TransportCoordinator = nullptr;
+	TSharedPtr<FWaveformEditorRenderData> RenderData = nullptr;
+
+	EWaveformEditorDisplayUnit DisplayUnit;
+};
+
diff --git a/Engine/Plugins/Editor/WaveformEditor/Source/WaveformEditorWidgets/Private/SWaveformPanel.cpp b/Engine/Plugins/Editor/WaveformEditor/Source/WaveformEditorWidgets/Private/SWaveformPanel.cpp
new file mode 100644
--- /dev/null
+++ b/Engine/Plugins/Editor/WaveformEditor/Source/WaveformEditorWidgets/Private/SWaveformPanel.cpp
@@ -0,0 +1,122 @@
+// Copyright Epic Games, Inc. All Rights Reserved.
+
+#include "SWaveformPanel.h"
+
+#include "SWaveformEditorTimeRuler.h"
+#include "SWaveformTransformationsOverlay.h"
+#include "SWaveformViewer.h"
+#include "SWaveformViewerOverlay.h"
+#include "WaveformEditorDisplayUnit.h"
+#include "WaveformEditorGridData.h"
+#include "WaveformEditorRenderData.h"
+#include "WaveformEditorTransportCoordinator.h"
+#include "WaveformEditorZoomController.h"
+#include "Widgets/SBoxPanel.h"
+#include "Widgets/SCompoundWidget.h"
+#include "Widgets/SOverlay.h"
+
+void SWaveformPanel::Construct(const FArguments& InArgs, TSharedRef<FWaveformEditorRenderData> InRenderData, TSharedRef<FWaveformEditorTransportCoordinator> InTransportCoordinator, TSharedRef<FWaveformEditorZoomController> InZoomManager, TSharedPtr<SWaveformTransformationsOverlay> InWaveformTransformationsOverlay)
+{
+	DisplayUnit = EWaveformEditorDisplayUnit::Seconds;
+
+	SetUpGridData(InRenderData, InTransportCoordinator);
+	SetUpWaveformViewer(InRenderData, InTransportCoordinator, InZoomManager, GridData.ToSharedRef());
+
+	if (InWaveformTransformationsOverlay)
+	{
+		WaveformTransformationsOverlay = InWaveformTransformationsOverlay;
+	}
+	
+	SetUpWaveformViewerOverlay(InTransportCoordinator, InZoomManager);
+	SetUpTimeRuler(InRenderData, InTransportCoordinator, GridData.ToSharedRef());
+	CreateLayout();
+}
+
+void SWaveformPanel::CreateLayout()
+{
+	check(TimeRuler);
+	check(WaveformViewer);
+	check(WaveformViewerOverlay);
+
+	TSharedPtr<SOverlay> WaveformView = SNew(SOverlay);
+	WaveformView->AddSlot()
+	[
+		WaveformViewer.ToSharedRef()
+	];
+
+	if (WaveformTransformationsOverlay)
+	{
+		WaveformView->AddSlot()
+		[
+			WaveformTransformationsOverlay.ToSharedRef()
+		];
+	}
+
+	WaveformView->AddSlot()
+	[
+		WaveformViewerOverlay.ToSharedRef()
+	];
+
+	ChildSlot
+	[
+		SNew(SVerticalBox)
+		+ SVerticalBox::Slot().AutoHeight()
+		[
+			TimeRuler.ToSharedRef()
+		]
+		+ SVerticalBox::Slot()
+		[
+			WaveformView.ToSharedRef()
+		]
+	];
+}
+
+void SWaveformPanel::SetUpTimeRuler(TSharedRef<FWaveformEditorRenderData> InRenderData, TSharedRef<FWaveformEditorTransportCoordinator> InTransportCoordinator, TSharedRef<FWaveformEditorGridData> InGridData)
+{
+	TimeRuler = SNew(SWaveformEditorTimeRuler, InTransportCoordinator, InRenderData).DisplayUnit(DisplayUnit);
+	InGridData->OnGridMetricsUpdated.AddSP(TimeRuler.ToSharedRef(), &SWaveformEditorTimeRuler::UpdateGridMetrics);
+	TimeRuler->OnTimeUnitMenuSelection.AddSP(this, &SWaveformPanel::UpdateDisplayUnit);
+}
+
+void SWaveformPanel::SetUpWaveformViewerOverlay(TSharedRef<FWaveformEditorTransportCoordinator> InTransportCoordinator, TSharedRef<FWaveformEditorZoomController> InZoomManager)
+{
+	WaveformViewerOverlay = SNew(SWaveformViewerOverlay, InTransportCoordinator);
+	WaveformViewerOverlay->OnNewMouseDelta.BindSP(InZoomManager, &FWaveformEditorZoomController::ZoomByDelta);
+}
+
+void SWaveformPanel::SetUpWaveformViewer(TSharedRef<FWaveformEditorRenderData> InRenderData, TSharedRef<FWaveformEditorTransportCoordinator> InTransportCoordinator, TSharedRef<FWaveformEditorZoomController> InZoomManager, TSharedRef<FWaveformEditorGridData> InGridData)
+{
+	WaveformViewer = SNew(SWaveformViewer, InRenderData, InTransportCoordinator);
+	InZoomManager->OnZoomRatioChanged.AddSP(InTransportCoordinator, &FWaveformEditorTransportCoordinator::OnZoomLevelChanged);
+	InGridData->OnGridMetricsUpdated.AddSP(WaveformViewer.ToSharedRef(), &SWaveformViewer::UpdateGridMetrics);
+}
+
+void SWaveformPanel::SetUpGridData(TSharedRef<FWaveformEditorRenderData> InRenderData, TSharedRef<FWaveformEditorTransportCoordinator> InTransportCoordinator)
+{
+	GridData = MakeShared<FWaveformEditorGridData>(InRenderData);
+	InTransportCoordinator->OnDisplayRangeUpdated.AddSP(GridData.ToSharedRef(), &FWaveformEditorGridData::UpdateDisplayRange);
+
+	const ISlateStyle* WaveEditorStyle = FSlateStyleRegistry::FindSlateStyle("WaveformEditorStyle");
+	if (ensure(WaveEditorStyle))
+	{
+		const FWaveformEditorTimeRulerStyle& RulerStyle = WaveEditorStyle->GetWidgetStyle<FWaveformEditorTimeRulerStyle>("WaveformEditorRuler.Style");
+		GridData->SetTicksTimeFont(&RulerStyle.TicksTextFont);
+	}
+}
+
+void SWaveformPanel::UpdateDisplayUnit(const EWaveformEditorDisplayUnit InDisplayUnit)
+{
+	DisplayUnit = InDisplayUnit;
+	TimeRuler->UpdateDisplayUnit(DisplayUnit);
+}
+
+void SWaveformPanel::Tick(const FGeometry& AllottedGeometry, const double InCurrentTime, const float InDeltaTime)
+{
+	const float PaintedWidth = AllottedGeometry.GetAbsoluteSize().X;
+
+	if (PaintedWidth != CachedPixelWidth)
+	{
+		GridData->UpdateGridMetrics(PaintedWidth);
+		CachedPixelWidth = PaintedWidth;
+	}
+}
\ No newline at end of file
diff --git a/Engine/Plugins/Editor/WaveformEditor/Source/WaveformEditorWidgets/Private/SWaveformTransformationsOverlay.cpp b/Engine/Plugins/Editor/WaveformEditor/Source/WaveformEditorWidgets/Private/SWaveformTransformationsOverlay.cpp
new file mode 100644
--- /dev/null
+++ b/Engine/Plugins/Editor/WaveformEditor/Source/WaveformEditorWidgets/Private/SWaveformTransformationsOverlay.cpp
@@ -0,0 +1,36 @@
+// Copyright Epic Games, Inc. All Rights Reserved.
+
+#include "SWaveformTransformationsOverlay.h"
+
+#include "Widgets/SOverlay.h"
+
+void SWaveformTransformationsOverlay::Construct(const FArguments& InArgs, TArrayView<TSharedPtr<SWidget>> InTransformationLayers)
+{
+	TransformationLayers = InTransformationLayers;
+	CreateLayout();
+}
+
+void SWaveformTransformationsOverlay::CreateLayout()
+{
+	ChildSlot
+	[
+		SAssignNew(MainOverlayPtr, SOverlay)
+	];
+
+	for (TSharedPtr<SWidget> Layer : TransformationLayers)
+	{
+		if (Layer)
+		{
+			MainOverlayPtr->AddSlot()
+			[
+				Layer.ToSharedRef()
+			];
+		}
+	}
+}
+
+void SWaveformTransformationsOverlay::OnLayerChainUpdate(TSharedPtr<SWidget>* FirstLayerPtr, const int32 NLayers)
+{
+	TransformationLayers = MakeArrayView(FirstLayerPtr, NLayers);
+	CreateLayout();
+}
diff --git a/Engine/Plugins/Editor/WaveformEditor/Source/WaveformEditorWidgets/Private/SWaveformViewer.cpp b/Engine/Plugins/Editor/WaveformEditor/Source/WaveformEditorWidgets/Private/SWaveformViewer.cpp
new file mode 100644
--- /dev/null
+++ b/Engine/Plugins/Editor/WaveformEditor/Source/WaveformEditorWidgets/Private/SWaveformViewer.cpp
@@ -0,0 +1,206 @@
+// Copyright Epic Games, Inc. All Rights Reserved.
+
+#include "SWaveformViewer.h"
+#include "WaveformEditorRenderData.h"
+
+void SWaveformViewer::Construct(const FArguments& InArgs, TSharedRef<FWaveformEditorRenderData> InRenderData, TSharedRef<FWaveformEditorTransportCoordinator> InTransportCoordinator)
+{
+	RenderData = InRenderData;
+	SampleData = RenderData->GetSampleData();
+
+	RenderData->OnRenderDataUpdated.AddSP(this, &SWaveformViewer::OnRenderDataUpdated);
+
+	TransportCoordinator = InTransportCoordinator;
+	TransportCoordinator->OnDisplayRangeUpdated.AddSP(this, &SWaveformViewer::OnDisplayRangeUpdated);
+
+	DisplayRange = TRange<float>::Inclusive(0.f, RenderData->GetOriginalWaveformDurationInSeconds());
+
+	Style = InArgs._Style;
+
+	check(Style)
+	WaveformColor = &Style->WaveformColor;
+	MajorGridLineColor = &Style->MajorGridLineColor;
+	MinorGridLineColor = &Style->MinorGridLineColor;
+	BackgroundColor = &Style->WaveformBackgroundColor;
+	BackgroundBrush = &Style->BackgroundBrush;
+	DesiredWidth = Style->DesiredWidth;
+	DesiredHeight = Style->DesiredHeight;
+
+}
+
+void SWaveformViewer::OnRenderDataUpdated()
+{
+	SampleData = RenderData->GetSampleData();
+	bForceRedraw = true;
+
+}
+
+void SWaveformViewer::OnDisplayRangeUpdated(const TRange<float> NewDisplayRange)
+{
+	const float LengthInSeconds = RenderData->GetOriginalWaveformDurationInSeconds();
+	DisplayRange.SetLowerBoundValue(NewDisplayRange.GetLowerBoundValue() * LengthInSeconds);
+	DisplayRange.SetUpperBoundValue(NewDisplayRange.GetUpperBoundValue() * LengthInSeconds);
+	bForceRedraw = true;
+}
+
+int32 SWaveformViewer::OnPaint(const FPaintArgs& Args, const FGeometry& AllottedGeometry, const FSlateRect& MyCullingRect, FSlateWindowElementList& OutDrawElements, int32 LayerId, const FWidgetStyle& InWidgetStyle, bool bParentEnabled) const 
+{
+	float PixelWidth = MyCullingRect.GetSize().X;
+
+	FSlateDrawElement::MakeBox(
+		OutDrawElements,
+		++LayerId,
+		AllottedGeometry.ToPaintGeometry(),
+		BackgroundBrush,
+		ESlateDrawEffect::None, 
+		BackgroundColor->GetSpecifiedColor()
+	);
+
+	if (PixelWidth > 0)
+	{
+		DrawGridLines(AllottedGeometry, OutDrawElements, LayerId);
+
+		TArray<FVector2D> WaveformPoints;
+ 		WaveformPoints.SetNumUninitialized(2);
+
+		for (const FWaveformLineCoordinates& PeakCoordinates : CachedDrawCoordinates)
+		{
+			WaveformPoints[0] = PeakCoordinates.PointA;
+			WaveformPoints[1] = PeakCoordinates.PointB;
+
+			FSlateDrawElement::MakeLines(
+				OutDrawElements,
+				LayerId,
+				AllottedGeometry.ToPaintGeometry(),
+				WaveformPoints,
+				ESlateDrawEffect::None, 
+				WaveformColor->GetSpecifiedColor()
+			);
+		}
+
+		++LayerId;
+	}
+
+	return LayerId;
+}
+
+FVector2D SWaveformViewer::ComputeDesiredSize(float) const
+{
+	return FVector2D(DesiredWidth, DesiredHeight);
+}
+
+void SWaveformViewer::Tick(const FGeometry& AllottedGeometry, const double InCurrentTime, const float InDeltaTime)
+{
+	uint32 DiscretePixelWidth = FMath::FloorToInt(AllottedGeometry.GetAbsoluteSize().X);
+	if (DiscretePixelWidth <= 0)
+	{
+		return;
+	}
+
+	if (DiscretePixelWidth != CachedPixelWidth || bForceRedraw)
+	{
+		CachedPixelWidth = DiscretePixelWidth;
+		CachedPixelHeight = AllottedGeometry.GetAbsoluteSize().Y;
+		WaveformDrawingUtils::GetBinnedPeaksFromWaveformRawData(CachedPeaks, CachedPixelWidth, SampleData.GetData(), RenderData->GetNumSamples(), RenderData->GetSampleRate(), RenderData->GetNumChannels(), DisplayRange.GetLowerBoundValue(), DisplayRange.GetUpperBoundValue());
+		GenerateWaveformLines(CachedDrawCoordinates, CachedPeaks, AllottedGeometry);
+		bForceRedraw = false;
+	}
+	else if (CachedPixelHeight != AllottedGeometry.GetAbsoluteSize().Y)
+	{
+		CachedPixelHeight = AllottedGeometry.GetAbsoluteSize().Y;
+		GenerateWaveformLines(CachedDrawCoordinates, CachedPeaks, AllottedGeometry);
+	}
+}
+
+void SWaveformViewer::GenerateWaveformLines(TArray<FWaveformLineCoordinates>& OutDrawCoordinates, const TArray<WaveformDrawingUtils::SampleRange>& InWaveformPeaks, const FGeometry& InAllottedGeometry, const float VerticalZoomFactor)
+{
+	const int NChannels = RenderData->GetNumChannels();
+
+	if (!ensure(NChannels != 0))
+	{
+		return;
+	}
+
+	const float ChannelSlotHeight = InAllottedGeometry.GetAbsoluteSize().Y / NChannels;
+	const float ChannelSlotMidPoint = ChannelSlotHeight / 2;
+	const float HeightScale = InAllottedGeometry.GetAbsoluteSize().Y / (2.f * 32767 * NChannels) * VerticalZoomFactor;
+	
+	const uint32 PixelWidth = FMath::FloorToInt(InAllottedGeometry.GetAbsoluteSize().X);
+	check(PixelWidth * NChannels == InWaveformPeaks.Num());
+
+	float MaxDistanceFromBoundary = 2.f;
+
+	OutDrawCoordinates.SetNumUninitialized(InWaveformPeaks.Num());
+	FWaveformLineCoordinates* OutCoordinatesData = OutDrawCoordinates.GetData();
+
+	for (uint16 Channel = 0; Channel < NChannels; ++Channel)
+	{
+		const float WaveformTopBoundary = ChannelSlotHeight * Channel;
+		const float WaveformBottomBoundary = WaveformTopBoundary + ChannelSlotHeight;
+		const float WaveformDrawCenter = WaveformTopBoundary + ChannelSlotMidPoint;
+
+		for (uint32 Pixel = 0; Pixel < PixelWidth; ++Pixel)
+		{
+			uint32 PeakIndex = Pixel * NChannels + Channel;
+			const WaveformDrawingUtils::SampleRange& SamplePeaks = InWaveformPeaks[PeakIndex];
+
+			const float SampleMaxScaled = SamplePeaks.GetUpperBoundValue() * HeightScale > 0.001f ? SamplePeaks.GetUpperBoundValue() * HeightScale : 0.1f;
+			const float SampleMinScaled = SamplePeaks.GetLowerBoundValue() * HeightScale < -0.001f ? SamplePeaks.GetLowerBoundValue() * HeightScale : -0.1f;
+
+			const float Top = FMath::Max(WaveformDrawCenter - SampleMaxScaled, WaveformTopBoundary + MaxDistanceFromBoundary);
+			const float Bottom = FMath::Min(WaveformDrawCenter - SampleMinScaled, WaveformBottomBoundary - MaxDistanceFromBoundary);
+
+			OutCoordinatesData[PeakIndex].PointA = FVector2D(Pixel, Top);
+			OutCoordinatesData[PeakIndex].PointB = FVector2D(Pixel, Bottom);
+		}
+	}
+}
+
+void SWaveformViewer::DrawGridLines(const FGeometry& AllottedGeometry, FSlateWindowElementList& OutDrawElements, int32& LayerId) const
+{
+	TArray<FVector2D> LinePoints;
+	LinePoints.SetNumUninitialized(2);
+
+	const double MinorGridXStep = GridMetrics.MajorGridXStep / GridMetrics.NumMinorGridDivisions;
+
+	for (double CurrentMajorLineX = GridMetrics.FirstMajorTickX; CurrentMajorLineX < AllottedGeometry.Size.X; CurrentMajorLineX += GridMetrics.MajorGridXStep)
+	{
+		const double MajorLineX = CurrentMajorLineX;
+
+		LinePoints[0] = FVector2D(MajorLineX, 0.f);
+		LinePoints[1] = FVector2D(MajorLineX, AllottedGeometry.Size.Y);
+
+		FSlateDrawElement::MakeLines(
+			OutDrawElements,
+			++LayerId,
+			AllottedGeometry.ToPaintGeometry(),
+			LinePoints,
+			ESlateDrawEffect::None,
+			MajorGridLineColor->GetSpecifiedColor(),
+			false);
+
+
+		for (int32 MinorLineIndex = 1; MinorLineIndex < GridMetrics.NumMinorGridDivisions; ++MinorLineIndex)
+		{
+			const double MinorLineX = MajorLineX + MinorGridXStep * MinorLineIndex;
+
+			LinePoints[0] = FVector2D(MinorLineX, 0.);
+			LinePoints[1] = FVector2D(MinorLineX, AllottedGeometry.Size.Y);
+
+			FSlateDrawElement::MakeLines(
+				OutDrawElements,
+				++LayerId,
+				AllottedGeometry.ToPaintGeometry(),
+				LinePoints,
+				ESlateDrawEffect::None,
+				MinorGridLineColor->GetSpecifiedColor(),
+				false);
+
+		}
+	}
+}
+
+void SWaveformViewer::UpdateGridMetrics(const FWaveEditorGridMetrics& InMetrics)
+{
+	GridMetrics = InMetrics;
+}
diff --git a/Engine/Plugins/Editor/WaveformEditor/Source/WaveformEditorWidgets/Private/SWaveformViewer.h b/Engine/Plugins/Editor/WaveformEditor/Source/WaveformEditorWidgets/Private/SWaveformViewer.h
new file mode 100644
--- /dev/null
+++ b/Engine/Plugins/Editor/WaveformEditor/Source/WaveformEditorWidgets/Private/SWaveformViewer.h
@@ -0,0 +1,73 @@
+// Copyright Epic Games, Inc. All Rights Reserved.
+
+#pragma once
+
+#include "Styling/SlateStyleRegistry.h"
+#include "Styling/SlateWidgetStyleAsset.h"
+#include "WaveformDrawingUtils.h"
+#include "WaveformEditorGridData.h"
+#include "WaveformEditorSlateTypes.h"
+#include "WaveformEditorTransportCoordinator.h"
+#include "Widgets/SLeafWidget.h"
+
+class FWaveformEditorRenderData;
+
+struct FWaveformLineCoordinates
+{
+	FVector2D PointA;
+	FVector2D PointB;
+};
+
+class SWaveformViewer : public SLeafWidget
+{
+public:
+
+	SLATE_BEGIN_ARGS(SWaveformViewer) 
+	{
+		const ISlateStyle* WaveEditorStyle = FSlateStyleRegistry::FindSlateStyle("WaveformEditorStyle");
+		check(WaveEditorStyle)
+		_Style = &WaveEditorStyle->GetWidgetStyle<FWaveformViewerStyle>("WaveformViewer.Style");
+	}
+
+	SLATE_STYLE_ARGUMENT(FWaveformViewerStyle, Style)
+		
+	SLATE_END_ARGS()
+
+	void Construct(const FArguments& InArgs, TSharedRef<FWaveformEditorRenderData> InRenderData, TSharedRef<FWaveformEditorTransportCoordinator> InTransportCoordinator);
+	
+	void OnRenderDataUpdated();
+	void OnDisplayRangeUpdated(const TRange<float> NewDisplayRange);
+	void UpdateGridMetrics(const FWaveEditorGridMetrics& InMetrics);
+
+private:
+	virtual int32 OnPaint(const FPaintArgs& Args, const FGeometry& AllottedGeometry, const FSlateRect& MyCullingRect, FSlateWindowElementList& OutDrawElements, int32 LayerId, const FWidgetStyle& InWidgetStyle, bool bParentEnabled) const override;
+	virtual FVector2D ComputeDesiredSize(float) const override;
+	virtual void Tick(const FGeometry& AllottedGeometry, const double InCurrentTime, const float InDeltaTime) override;
+
+	void GenerateWaveformLines(TArray<FWaveformLineCoordinates>& OutDrawCoordinates, const TArray<WaveformDrawingUtils::SampleRange>& InWaveformPeaks, const FGeometry& InAllottedGeometry, const float VerticalZoomFactor = 0.9f);
+	void DrawGridLines(const FGeometry& AllottedGeometry, FSlateWindowElementList& OutDrawElements, int32& LayerId) const;
+
+	TRange<float> DisplayRange;
+	bool bForceRedraw = false;
+
+	TArrayView<const int16> SampleData;
+	FWaveEditorGridMetrics GridMetrics;
+
+	const FWaveformViewerStyle* Style = nullptr;
+	const FSlateBrush* BackgroundBrush = nullptr;
+	const FSlateColor* BackgroundColor = nullptr;
+	const FSlateColor* WaveformColor = nullptr;
+	const FSlateColor* MajorGridLineColor = nullptr;
+	const FSlateColor* MinorGridLineColor = nullptr;
+	float DesiredHeight = 0.f;
+	float DesiredWidth = 0.f;
+
+	uint32 CachedPixelWidth = 0; 
+	float CachedPixelHeight = 0.f;
+	TArray<WaveformDrawingUtils::SampleRange> CachedPeaks;
+	TArray<FWaveformLineCoordinates> CachedDrawCoordinates;
+
+
+	TSharedPtr<FWaveformEditorTransportCoordinator> TransportCoordinator = nullptr;
+	TSharedPtr<FWaveformEditorRenderData> RenderData = nullptr;
+};
\ No newline at end of file
diff --git a/Engine/Plugins/Editor/WaveformEditor/Source/WaveformEditorWidgets/Private/SWaveformViewerOverlay.cpp b/Engine/Plugins/Editor/WaveformEditor/Source/WaveformEditorWidgets/Private/SWaveformViewerOverlay.cpp
new file mode 100644
--- /dev/null
+++ b/Engine/Plugins/Editor/WaveformEditor/Source/WaveformEditorWidgets/Private/SWaveformViewerOverlay.cpp
@@ -0,0 +1,67 @@
+// Copyright Epic Games, Inc. All Rights Reserved.
+
+#include "SWaveformViewerOverlay.h"
+
+#include "WaveformEditorTransportCoordinator.h"
+#include "Widgets/SLeafWidget.h"
+
+void SWaveformViewerOverlay::Construct(const FArguments& InArgs, TSharedRef<FWaveformEditorTransportCoordinator> InTransportCoordinator)
+{
+	TransportCoordinator = InTransportCoordinator;
+
+	Style = InArgs._Style;
+
+	check(Style)
+	PlayheadColor = &Style->PlayheadColor;
+	PlayheadWidth = Style->PlayheadWidth;
+	DesiredWidth = Style->DesiredWidth;
+	DesiredHeight = Style->DesiredHeight;
+}
+
+FReply SWaveformViewerOverlay::OnMouseButtonDown(const FGeometry& MyGeometry, const FPointerEvent& MouseEvent)
+{
+	return TransportCoordinator->ReceiveMouseButtonDown(*this, MyGeometry, MouseEvent);
+}
+
+FReply SWaveformViewerOverlay::OnMouseWheel(const FGeometry& MyGeometry, const FPointerEvent& MouseEvent)
+{
+	OnNewMouseDelta.ExecuteIfBound(MouseEvent.GetWheelDelta());
+	return FReply::Handled();
+}
+
+int32 SWaveformViewerOverlay::OnPaint(const FPaintArgs& Args, const FGeometry& AllottedGeometry, const FSlateRect& MyCullingRect, FSlateWindowElementList& OutDrawElements, int32 LayerId, const FWidgetStyle& InWidgetStyle, bool bParentEnabled) const
+{
+	LayerId = DrawPlayhead(AllottedGeometry, OutDrawElements, LayerId);
+
+	return LayerId;
+}
+
+int32 SWaveformViewerOverlay::DrawPlayhead(const FGeometry& AllottedGeometry, FSlateWindowElementList& OutDrawElements, int32 LayerId) const
+{
+	const float PlayHeadX = AllottedGeometry.GetLocalSize().X * TransportCoordinator->GetPlayheadPosition();
+
+	TArray<FVector2D> LinePoints;
+	{
+		LinePoints.AddUninitialized(2);
+		LinePoints[0] = FVector2D(PlayHeadX, 0.0f);
+		LinePoints[1] = FVector2D(PlayHeadX, AllottedGeometry.Size.Y);
+	}
+
+	FSlateDrawElement::MakeLines(
+		OutDrawElements,
+		LayerId,
+		AllottedGeometry.ToPaintGeometry(),
+		LinePoints,
+		ESlateDrawEffect::None,
+		PlayheadColor->GetSpecifiedColor(),
+		true, 
+		PlayheadWidth
+	);
+
+	return ++LayerId;
+}
+
+FVector2D SWaveformViewerOverlay::ComputeDesiredSize(float) const
+{
+	return FVector2D(DesiredWidth, DesiredHeight);
+}
\ No newline at end of file
diff --git a/Engine/Plugins/Editor/WaveformEditor/Source/WaveformEditorWidgets/Private/SWaveformViewerOverlay.h b/Engine/Plugins/Editor/WaveformEditor/Source/WaveformEditorWidgets/Private/SWaveformViewerOverlay.h
new file mode 100644
--- /dev/null
+++ b/Engine/Plugins/Editor/WaveformEditor/Source/WaveformEditorWidgets/Private/SWaveformViewerOverlay.h
@@ -0,0 +1,49 @@
+// Copyright Epic Games, Inc. All Rights Reserved.
+
+#pragma once
+
+#include "Styling/SlateStyleRegistry.h"
+#include "Styling/SlateWidgetStyleAsset.h"
+#include "WaveformEditorSlateTypes.h"
+#include "WaveformEditorTransportCoordinator.h"
+#include "Widgets/SLeafWidget.h"
+
+DECLARE_DELEGATE_OneParam(FOnNewMouseDelta, const float /* new delta */)
+
+class SWaveformViewerOverlay : public SLeafWidget
+{
+public:
+	SLATE_BEGIN_ARGS(SWaveformViewerOverlay)
+	{
+		const ISlateStyle* WaveEditorStyle = FSlateStyleRegistry::FindSlateStyle("WaveformEditorStyle");
+		check(WaveEditorStyle)
+		_Style = &WaveEditorStyle->GetWidgetStyle<FWaveformViewerOverlayStyle>("WaveformViewerOverlay.Style");
+
+	}
+
+	SLATE_STYLE_ARGUMENT(FWaveformViewerOverlayStyle, Style)
+
+	SLATE_END_ARGS()
+
+	void Construct(const FArguments& InArgs, TSharedRef<FWaveformEditorTransportCoordinator> InTransportCoordinator);
+
+	virtual FReply OnMouseButtonDown(const FGeometry& MyGeometry, const FPointerEvent& MouseEvent) override;
+	virtual FReply OnMouseWheel(const FGeometry& MyGeometry, const FPointerEvent& MouseEvent) override;
+
+	FOnNewMouseDelta OnNewMouseDelta;
+
+private:
+	virtual int32 OnPaint(const FPaintArgs& Args, const FGeometry& AllottedGeometry, const FSlateRect& MyCullingRect, FSlateWindowElementList& OutDrawElements, int32 LayerId, const FWidgetStyle& InWidgetStyle, bool bParentEnabled) const override;
+	virtual FVector2D ComputeDesiredSize(float) const override;
+
+	int32 DrawPlayhead(const FGeometry& AllottedGeometry, FSlateWindowElementList& OutDrawElements, int32 LayerId) const;
+
+	TSharedPtr<FWaveformEditorTransportCoordinator> TransportCoordinator = nullptr;
+
+	const FWaveformViewerOverlayStyle* Style = nullptr;
+
+	const FSlateColor* PlayheadColor = nullptr;
+	float PlayheadWidth = 0.f;
+	float DesiredWidth = 0.f;
+	float DesiredHeight = 0.f;
+};
\ No newline at end of file
diff --git a/Engine/Plugins/Editor/WaveformEditor/Source/WaveformEditorWidgets/Private/WaveformDrawingUtils.cpp b/Engine/Plugins/Editor/WaveformEditor/Source/WaveformEditorWidgets/Private/WaveformDrawingUtils.cpp
new file mode 100644
--- /dev/null
+++ b/Engine/Plugins/Editor/WaveformEditor/Source/WaveformEditorWidgets/Private/WaveformDrawingUtils.cpp
@@ -0,0 +1,40 @@
+// Copyright Epic Games, Inc. All Rights Reserved.
+
+#include "WaveformDrawingUtils.h"
+
+void WaveformDrawingUtils::GetBinnedPeaksFromWaveformRawData(TArray<SampleRange>& OutWaveformPeaks, const uint32 NumBins, const int16* RawDataPtr, const uint32 NumSamples, const uint32 SampleRate, const uint16 NChannels, const float StartTime, const float EndTime)
+{
+	uint32 NumPeaks = NumBins * NChannels;
+	OutWaveformPeaks.SetNumUninitialized(NumPeaks);
+
+	double TimePerBin = (EndTime - StartTime) / NumBins;
+	double IterationStartTime = StartTime;
+	uint32 FirstBinnedFrame = FMath::RoundToInt(IterationStartTime * SampleRate);
+
+	for (uint32 Peak = 0; Peak < NumPeaks; Peak += NChannels)
+	{
+		uint32 LastBinnedFrame = FMath::RoundToInt((IterationStartTime + TimePerBin) * SampleRate);
+
+		for (uint16 Channel = 0; Channel < NChannels; ++Channel)
+		{
+			int16 MaxSampleValue = TNumericLimits<int16>::Min();
+			int16 MinSampleValue = TNumericLimits<int16>::Max();
+
+			for (uint32 Frame = FirstBinnedFrame; Frame < LastBinnedFrame; ++Frame)
+			{
+				const uint32 SampleIndex = Frame * NChannels + Channel;
+				check(SampleIndex < NumSamples);
+				int16 SampleValue = RawDataPtr[SampleIndex];
+
+				MinSampleValue = FMath::Min(MinSampleValue, SampleValue);
+				MaxSampleValue = FMath::Max(MaxSampleValue, SampleValue);
+
+			}
+
+			OutWaveformPeaks[Peak + Channel] = SampleRange(MinSampleValue, MaxSampleValue);
+		}
+
+		IterationStartTime += TimePerBin;
+		FirstBinnedFrame = LastBinnedFrame;
+	}
+}
diff --git a/Engine/Plugins/Editor/WaveformEditor/Source/WaveformEditorWidgets/Private/WaveformEditorDisplayUnit.h b/Engine/Plugins/Editor/WaveformEditor/Source/WaveformEditorWidgets/Private/WaveformEditorDisplayUnit.h
new file mode 100644
--- /dev/null
+++ b/Engine/Plugins/Editor/WaveformEditor/Source/WaveformEditorWidgets/Private/WaveformEditorDisplayUnit.h
@@ -0,0 +1,9 @@
+// Copyright Epic Games, Inc. All Rights Reserved.
+
+#pragma once
+
+enum class EWaveformEditorDisplayUnit
+{
+	Samples,
+	Seconds
+};
\ No newline at end of file
diff --git a/Engine/Plugins/Editor/WaveformEditor/Source/WaveformEditorWidgets/Private/WaveformEditorGridData.cpp b/Engine/Plugins/Editor/WaveformEditor/Source/WaveformEditorWidgets/Private/WaveformEditorGridData.cpp
new file mode 100644
--- /dev/null
+++ b/Engine/Plugins/Editor/WaveformEditor/Source/WaveformEditorWidgets/Private/WaveformEditorGridData.cpp
@@ -0,0 +1,72 @@
+#include "WaveformEditorGridData.h"
+
+#include "Fonts/FontMeasure.h"
+#include "Framework/Application/SlateApplication.h"
+#include "GenericPlatform/GenericPlatformMath.h"
+#include "Misc/FrameRate.h"
+#include "WaveformEditorRenderData.h"
+
+FWaveformEditorGridData::FWaveformEditorGridData(TSharedRef<FWaveformEditorRenderData> InRenderData, const FSlateFontInfo* InTicksTimeFont /* = nullptr*/)
+	: RenderData(InRenderData)
+	, TicksTimeFont(InTicksTimeFont)
+{
+}
+
+void FWaveformEditorGridData::UpdateDisplayRange(const TRange<float> InDisplayRange)
+{
+	DisplayRange = InDisplayRange;
+	UpdateGridMetrics(GridPixelWidth);
+}
+
+bool FWaveformEditorGridData::UpdateGridMetrics(const float InGridPixelWidth)
+{
+	if (GridPixelWidth != InGridPixelWidth)
+	{
+		GridPixelWidth = InGridPixelWidth;
+	}
+	
+	const float WaveformDurationSeconds = RenderData->GetOriginalWaveformDurationInSeconds();
+	const double StartTime = WaveformDurationSeconds * DisplayRange.GetLowerBoundValue();
+	const double DisplayedDuration = WaveformDurationSeconds * DisplayRange.Size<double>();
+	const double PixelsPerSecond = GridPixelWidth / DisplayedDuration;
+	double MajorGridStepSeconds = 0.0;
+	float MinTicksPixelSpacing = 30.0f;
+	
+	if (TicksTimeFont)
+	{
+		TSharedRef<FSlateFontMeasure> FontMeasureService = FSlateApplication::Get().GetRenderer()->GetFontMeasureService();
+		FString LongestDisplayedTimeString = FString::Printf(TEXT("%.3f"), WaveformDurationSeconds * DisplayRange.GetUpperBoundValue());
+		FVector2D MaxTextSize = FontMeasureService->Measure(LongestDisplayedTimeString, *TicksTimeFont);
+		MinTicksPixelSpacing = MaxTextSize.X;
+	}
+
+	const FFrameRate FrameRate(RenderData->GetSampleRate(), 1);
+
+	if (!FrameRate.ComputeGridSpacing(PixelsPerSecond, MajorGridStepSeconds, GridMetrics.NumMinorGridDivisions, MinTicksPixelSpacing + 5.f))
+	{
+		return false;
+	}
+
+	const double GridOffset = FGenericPlatformMath::Fmod(StartTime, MajorGridStepSeconds);
+	GridMetrics.FirstMajorTickX = (0.f - GridOffset) * PixelsPerSecond;
+	GridMetrics.MajorGridXStep = MajorGridStepSeconds * PixelsPerSecond;
+	GridMetrics.StartTime = StartTime;
+	GridMetrics.PixelsPerSecond = PixelsPerSecond;
+
+	if (OnGridMetricsUpdated.IsBound())
+	{
+		OnGridMetricsUpdated.Broadcast(GridMetrics);
+	}
+
+	return true;
+}
+
+const FWaveEditorGridMetrics FWaveformEditorGridData::GetGridMetrics() const
+{
+	return GridMetrics;
+}
+
+void FWaveformEditorGridData::SetTicksTimeFont(const FSlateFontInfo* InNewFont)
+{
+	TicksTimeFont = InNewFont;
+}
diff --git a/Engine/Plugins/Editor/WaveformEditor/Source/WaveformEditorWidgets/Private/WaveformEditorRenderData.cpp b/Engine/Plugins/Editor/WaveformEditor/Source/WaveformEditorWidgets/Private/WaveformEditorRenderData.cpp
new file mode 100644
--- /dev/null
+++ b/Engine/Plugins/Editor/WaveformEditor/Source/WaveformEditorWidgets/Private/WaveformEditorRenderData.cpp
@@ -0,0 +1,67 @@
+// Copyright Epic Games, Inc. All Rights Reserved.
+
+#include "WaveformEditorRenderData.h"
+
+#include "AudioExtensions/Public/IWaveformTransformation.h"
+#include "DSP/BufferVectorOperations.h"
+#include "DSP/FloatArrayMath.h"
+#include "Sound/SoundWave.h"
+
+TArrayView<const int16> FWaveformEditorRenderData::GetSampleData() const
+{
+	return SampleData;
+}
+
+const uint16 FWaveformEditorRenderData::GetNumChannels() const
+{
+	return NumChannels;
+}
+
+const uint32 FWaveformEditorRenderData::GetNumSamples() const
+{
+	return NumSamples;
+}
+
+const uint32 FWaveformEditorRenderData::GetSampleRate() const
+{
+	return SampleRate;
+}
+
+const float FWaveformEditorRenderData::GetOriginalWaveformDurationInSeconds() const
+{
+	return OriginalWaveformDurationInSeconds;
+}
+
+const uint32 FWaveformEditorRenderData::GetOriginalWaveformFrames() const
+{
+	return OriginalWaveformFrames;
+}
+
+const TRange<float> FWaveformEditorRenderData::GetTransformedWaveformBounds() const
+{
+	return TransformedWaveformBounds;
+}
+
+void FWaveformEditorRenderData::UpdateRenderData(const uint8* InRawData, const uint32 InNumSamples, const uint32 InFirstEditedSample, const uint32 InLastEditedSample, const uint32 InSampleRate, const uint16 InNumChannels)
+{
+	NumSamples = InNumSamples;
+	NumEditedSamples = InLastEditedSample - InFirstEditedSample;
+	SampleRate = InSampleRate;
+	NumChannels = InNumChannels;
+	const int16* SampleDataPtr = reinterpret_cast<const int16*>(InRawData);
+	SampleData = MakeArrayView(SampleDataPtr, NumSamples);
+
+	OriginalWaveformDurationInSeconds = NumSamples / NumChannels / SampleRate;
+	OriginalWaveformFrames = NumSamples / NumChannels;
+
+	const float FirstEditedSampleOffset = InFirstEditedSample / (float) NumSamples;
+	const float LastEditedSampleOffset = (NumSamples - InLastEditedSample) / (float) NumSamples;	
+	
+	TransformedWaveformBounds.SetLowerBoundValue(FirstEditedSampleOffset);
+	TransformedWaveformBounds.SetUpperBoundValue(1 - LastEditedSampleOffset);
+
+	if (OnRenderDataUpdated.IsBound())
+	{
+		OnRenderDataUpdated.Broadcast();
+	}
+}
\ No newline at end of file
diff --git a/Engine/Plugins/Editor/WaveformEditor/Source/WaveformEditorWidgets/Private/WaveformEditorSlateTypes.cpp b/Engine/Plugins/Editor/WaveformEditor/Source/WaveformEditorWidgets/Private/WaveformEditorSlateTypes.cpp
new file mode 100644
--- /dev/null
+++ b/Engine/Plugins/Editor/WaveformEditor/Source/WaveformEditorWidgets/Private/WaveformEditorSlateTypes.cpp
@@ -0,0 +1,88 @@
+// Copyright Epic Games, Inc. All Rights Reserved.
+
+#include "WaveformEditorSlateTypes.h"
+
+#include "Styling/AppStyle.h"
+#include "Styling/StyleColors.h"
+
+namespace WaveformEditorStylesSharedParams
+{
+	const FLazyName BackgroundBrushName = TEXT("WhiteBrush");
+	const FLazyName HandleBrushName = TEXT("Sequencer.Timeline.VanillaScrubHandleDown");
+	const FLinearColor PlayheadColor = FLinearColor(255, 0.1, 0.2, 1.0);
+	const FLinearColor RulerTicksColor = FLinearColor(1.f, 1.f, 1.f, 0.9);
+	const FLinearColor WaveformBackGroundColor = FLinearColor(0.02, 0.02, 0.02, 1.f);
+	const float ViewerHeight = 720.f;
+	const float ViewerWidth = 1280.f;
+}
+
+const FName FWaveformViewerStyle::TypeName("FWaveformViewerStyle");
+
+FWaveformViewerStyle::FWaveformViewerStyle()
+	: WaveformColor(FLinearColor::White)
+	, MajorGridLineColor(FLinearColor::Black)
+	, MinorGridLineColor(FLinearColor(0.f, 0.f, 0.f, 0.5))
+	, WaveformBackgroundColor(WaveformEditorStylesSharedParams::WaveformBackGroundColor)
+	, BackgroundBrush(*FAppStyle::GetBrush(WaveformEditorStylesSharedParams::BackgroundBrushName))
+	, DesiredWidth(WaveformEditorStylesSharedParams::ViewerWidth)
+	, DesiredHeight(WaveformEditorStylesSharedParams::ViewerHeight)
+{
+}
+
+const FWaveformViewerStyle& FWaveformViewerStyle::GetDefault()
+{
+	static FWaveformViewerStyle Default;
+	return Default;
+}
+
+void FWaveformViewerStyle::GetResources(TArray< const FSlateBrush* >& OutBrushes) const
+{
+	OutBrushes.Add(&BackgroundBrush);
+}
+
+const FName FWaveformViewerOverlayStyle::TypeName("FWaveformViewerOverlayStyle");
+
+FWaveformViewerOverlayStyle::FWaveformViewerOverlayStyle()
+	: PlayheadColor(WaveformEditorStylesSharedParams::PlayheadColor)
+	, PlayheadWidth(1.0f)
+	, DesiredWidth(WaveformEditorStylesSharedParams::ViewerWidth)
+	, DesiredHeight(WaveformEditorStylesSharedParams::ViewerHeight)
+{
+}
+
+const FWaveformViewerOverlayStyle& FWaveformViewerOverlayStyle::GetDefault()
+{
+	static FWaveformViewerOverlayStyle Default;
+	return Default;
+}
+
+
+const FName FWaveformEditorTimeRulerStyle::TypeName("FWaveformEditorTimeRulerStyle");
+
+FWaveformEditorTimeRulerStyle::FWaveformEditorTimeRulerStyle()
+	: HandleWidth(15.f)
+	, HandleColor(WaveformEditorStylesSharedParams::PlayheadColor)
+	, HandleBrush(*FAppStyle::GetBrush(WaveformEditorStylesSharedParams::HandleBrushName))
+	, TicksColor(WaveformEditorStylesSharedParams::RulerTicksColor)
+	, TicksTextColor(WaveformEditorStylesSharedParams::RulerTicksColor)
+	, TicksTextFont(FAppStyle::GetFontStyle("Regular"))
+	, TicksTextOffset(5.f)
+	, BackgroundColor(FLinearColor::Black)
+	, BackgroundBrush(*FAppStyle::GetBrush(WaveformEditorStylesSharedParams::BackgroundBrushName))
+	, DesiredWidth(WaveformEditorStylesSharedParams::ViewerWidth)
+	, DesiredHeight(30.f)
+{
+}
+
+const FWaveformEditorTimeRulerStyle& FWaveformEditorTimeRulerStyle::GetDefault()
+{
+	static FWaveformEditorTimeRulerStyle Default;
+	return Default;
+}
+
+void FWaveformEditorTimeRulerStyle::GetResources(TArray< const FSlateBrush* >& OutBrushes) const
+{
+	OutBrushes.Add(&HandleBrush);
+	OutBrushes.Add(&BackgroundBrush);
+}
+
diff --git a/Engine/Plugins/Editor/WaveformEditor/Source/WaveformEditorWidgets/Private/WaveformEditorStyle.cpp b/Engine/Plugins/Editor/WaveformEditor/Source/WaveformEditorWidgets/Private/WaveformEditorStyle.cpp
new file mode 100644
--- /dev/null
+++ b/Engine/Plugins/Editor/WaveformEditor/Source/WaveformEditorWidgets/Private/WaveformEditorStyle.cpp
@@ -0,0 +1,127 @@
+// Copyright Epic Games, Inc. All Rights Reserved.
+
+#include "WaveformEditorStyle.h"
+
+#include "Styling/SlateStyleRegistry.h"
+#include "WaveformEditorSlateTypes.h"
+#include "WaveformEditorWidgetsSettings.h"
+#include "WaveformEditorWidgetsSettings.h"
+
+static FLazyName WaveformViewerStyleName("WaveformViewer.Style");
+static FLazyName WaveformViewerOverlayStyleName("WaveformViewerOverlay.Style");
+static FLazyName WaveformEditorRulerStyleName("WaveformEditorRuler.Style");
+
+FName FWaveformEditorStyle::StyleName("WaveformEditorStyle");
+TUniquePtr<FWaveformEditorStyle> FWaveformEditorStyle::StyleInstance = nullptr;
+
+FWaveformEditorStyle::FWaveformEditorStyle()
+	: FSlateStyleSet(StyleName)
+{
+}
+
+FWaveformEditorStyle::~FWaveformEditorStyle()
+{
+	FSlateStyleRegistry::UnRegisterSlateStyle(*this);
+}
+
+const FWaveformEditorStyle& FWaveformEditorStyle::Get()
+{
+	check(StyleInstance)
+	return *StyleInstance;
+}
+
+void FWaveformEditorStyle::Init()
+{
+	if (StyleInstance == nullptr)
+	{
+		StyleInstance = MakeUnique<FWaveformEditorStyle>();
+	}
+
+	const UWaveformEditorWidgetsSettings* Settings = GetWidgetsSettings();
+	
+	check(Settings);
+	Settings->OnSettingChanged().AddStatic(&FWaveformEditorStyle::OnWidgetSettingsUpdated);
+
+	StyleInstance->SetParentStyleName("CoreStyle");
+
+	//Waveform Viewer style
+	FWaveformViewerStyle WaveViewerStyle = FWaveformViewerStyle()
+		.SetWaveformColor(Settings->WaveformColor)
+		.SetBackgroundColor(Settings->WaveformBackgroundColor)
+		.SetMajorGridLineColor(Settings->MajorGridColor)
+		.SetMinorGridLineColor(Settings->MinorGridColor);
+
+	StyleInstance->Set(WaveformViewerStyleName, WaveViewerStyle);
+
+	//Waveform Viewer Overlay style
+	FWaveformViewerOverlayStyle WaveViewerOverlayStyle = FWaveformViewerOverlayStyle().SetPlayheadColor(Settings->PlayheadColor);
+	StyleInstance->Set(WaveformViewerOverlayStyleName, WaveViewerOverlayStyle);
+
+	//Time Ruler style 
+	FWaveformEditorTimeRulerStyle TimeRulerStyle = FWaveformEditorTimeRulerStyle()
+		.SetHandleColor(Settings->PlayheadColor)
+		.SetTicksColor(Settings->RulerTicksColor)
+		.SetTicksTextColor(Settings->RulerTextColor)
+		.SetHandleColor(Settings->PlayheadColor)
+		.SetFontSize(Settings->RulerFontSize);
+
+	StyleInstance->Set(WaveformEditorRulerStyleName, TimeRulerStyle);
+
+	FSlateStyleRegistry::RegisterSlateStyle(StyleInstance->Get());
+
+}
+
+const UWaveformEditorWidgetsSettings* FWaveformEditorStyle::GetWidgetsSettings() 
+{
+	const UWaveformEditorWidgetsSettings* WaveformEditorWidgetsSettings = GetDefault<UWaveformEditorWidgetsSettings>();
+	check(WaveformEditorWidgetsSettings)
+
+	return WaveformEditorWidgetsSettings;
+}
+
+void FWaveformEditorStyle::OnWidgetSettingsUpdated(const FName& PropertyName, const UWaveformEditorWidgetsSettings* Settings)
+{
+	TSharedRef<FWaveformViewerStyle> WaveformViewerStyle = GetRegisteredWidgetStyle<FWaveformViewerStyle>(WaveformViewerStyleName);
+	TSharedRef<FWaveformViewerOverlayStyle> WaveformViewerOverlayStyle = GetRegisteredWidgetStyle<FWaveformViewerOverlayStyle>(WaveformViewerOverlayStyleName);
+	TSharedRef<FWaveformEditorTimeRulerStyle> WaveformEditorTimeRulerStyle = GetRegisteredWidgetStyle<FWaveformEditorTimeRulerStyle>(WaveformEditorRulerStyleName);
+
+	if (PropertyName == GET_MEMBER_NAME_CHECKED(UWaveformEditorWidgetsSettings, WaveformColor))
+	{
+		WaveformViewerStyle->SetWaveformColor(Settings->WaveformColor);
+	}
+	else if (PropertyName == GET_MEMBER_NAME_CHECKED(UWaveformEditorWidgetsSettings, WaveformBackgroundColor))
+	{
+		WaveformViewerStyle->SetBackgroundColor(Settings->WaveformBackgroundColor);
+	}
+	else if (PropertyName == GET_MEMBER_NAME_CHECKED(UWaveformEditorWidgetsSettings, MajorGridColor))
+	{
+		WaveformViewerStyle->SetMajorGridLineColor(Settings->MajorGridColor);
+	}
+	else if (PropertyName == GET_MEMBER_NAME_CHECKED(UWaveformEditorWidgetsSettings, MinorGridColor))
+	{
+		WaveformViewerStyle->SetMinorGridLineColor(Settings->MinorGridColor);
+	}
+	else if (PropertyName == GET_MEMBER_NAME_CHECKED(UWaveformEditorWidgetsSettings, PlayheadColor))
+	{
+		WaveformViewerOverlayStyle->SetPlayheadColor(Settings->PlayheadColor);
+		WaveformEditorTimeRulerStyle->SetHandleColor(Settings->PlayheadColor);
+	}
+	else if (PropertyName == GET_MEMBER_NAME_CHECKED(UWaveformEditorWidgetsSettings, RulerBackgroundColor))
+	{
+		WaveformEditorTimeRulerStyle->SetBackgroundColor(Settings->RulerBackgroundColor);
+	}
+	else if (PropertyName == GET_MEMBER_NAME_CHECKED(UWaveformEditorWidgetsSettings, RulerTicksColor))
+	{
+		WaveformEditorTimeRulerStyle->SetTicksColor(Settings->RulerTicksColor);
+	}
+	else if (PropertyName == GET_MEMBER_NAME_CHECKED(UWaveformEditorWidgetsSettings, RulerTextColor))
+	{
+		WaveformEditorTimeRulerStyle->SetTicksTextColor(Settings->RulerTextColor);
+	}
+	else if (PropertyName == GET_MEMBER_NAME_CHECKED(UWaveformEditorWidgetsSettings, RulerFontSize))
+	{
+		WaveformEditorTimeRulerStyle->SetFontSize(Settings->RulerFontSize);
+	}
+}
+
+
diff --git a/Engine/Plugins/Editor/WaveformEditor/Source/WaveformEditorWidgets/Private/WaveformEditorTransportCoordinator.cpp b/Engine/Plugins/Editor/WaveformEditor/Source/WaveformEditorWidgets/Private/WaveformEditorTransportCoordinator.cpp
new file mode 100644
--- /dev/null
+++ b/Engine/Plugins/Editor/WaveformEditor/Source/WaveformEditorWidgets/Private/WaveformEditorTransportCoordinator.cpp
@@ -0,0 +1,298 @@
+// Copyright Epic Games, Inc. All Rights Reserved.
+
+#include "WaveformEditorTransportCoordinator.h"
+
+#include "WaveformEditorRenderData.h"
+
+/** Zoom Formula
+E.g: I am receiving a playback ratio (PBR) in the range 0 to 1
+To position the playhead accordingly, this is the formula to follow:
+
+PH = (PBR - LB) / ZR
+
+PH = Playhead Position
+ZR = Zoom Ratio
+LB = Lower Displayed Waveform Bound
+
+from here, you can rotate as needed. 
+e.g to find PBR
+PBR = (PH * ZR) + LB
+
+This formula can be used to position any element (e.g. a marker) given in a 0 to 1 position range where
+0 is the beginning time of the waveform and 1 the end time
+*/
+
+FWaveformEditorTransportCoordinator::FWaveformEditorTransportCoordinator(TSharedRef<FWaveformEditorRenderData> InRenderData)
+	: RenderData(InRenderData)
+{
+}
+
+FReply FWaveformEditorTransportCoordinator::ReceiveMouseButtonDown(SWidget& WidgetOwner, const FGeometry& Geometry, const FPointerEvent& MouseEvent)
+{
+	if (WidgetOwner.GetTypeAsString() == TEXT("SWaveformViewerOverlay"))
+	{
+		return HandleWaveformViewerOverlayInteraction(MouseEvent, Geometry);
+	}
+
+	if (WidgetOwner.GetTypeAsString() == TEXT("SWaveformEditorTimeRuler"))
+	{
+		return HandleTimeRulerInteraction(EReceivedInteractionType::MouseButtonDown, MouseEvent, WidgetOwner, Geometry);
+	}
+
+	return FReply::Unhandled();
+}
+
+FReply FWaveformEditorTransportCoordinator::ReceiveMouseButtonUp(SWidget& WidgetOwner, const FGeometry& Geometry, const FPointerEvent& MouseEvent)
+{
+	if (WidgetOwner.GetTypeAsString() == TEXT("SWaveformEditorTimeRuler"))
+	{
+		return HandleTimeRulerInteraction(EReceivedInteractionType::MouseButtonUp, MouseEvent, WidgetOwner, Geometry);
+	}
+
+	return FReply::Unhandled();
+}
+
+FReply FWaveformEditorTransportCoordinator::ReceiveMouseMove(SWidget& WidgetOwner, const FGeometry& Geometry, const FPointerEvent& MouseEvent)
+{
+	if (WidgetOwner.GetTypeAsString() == TEXT("SWaveformEditorTimeRuler"))
+	{
+		return HandleTimeRulerInteraction(EReceivedInteractionType::MouseMove, MouseEvent, WidgetOwner, Geometry);
+	}
+
+	return FReply::Unhandled();
+}
+
+FReply FWaveformEditorTransportCoordinator::HandleWaveformViewerOverlayInteraction(const FPointerEvent& MouseEvent, const FGeometry& Geometry)
+{
+	const bool HandleLeftMouseButton = MouseEvent.GetEffectingButton() == EKeys::LeftMouseButton;
+	if (HandleLeftMouseButton)
+	{
+		const float LocalWidth = Geometry.GetLocalSize().X;
+
+		if (LocalWidth > 0.f)
+		{
+			const float NewPosition = Geometry.AbsoluteToLocal(MouseEvent.GetScreenSpacePosition()).X / LocalWidth;
+			ScrubPlayhead(NewPosition, false);
+		}
+	}
+
+	return FReply::Handled();
+}
+
+FReply FWaveformEditorTransportCoordinator::HandleTimeRulerInteraction(const EReceivedInteractionType MouseInteractionType, const FPointerEvent& MouseEvent, SWidget& TimeRulerWidget, const FGeometry& Geometry)
+{
+	const float LocalWidth = Geometry.GetLocalSize().X;
+
+	if (LocalWidth > 0.f)
+	{
+		const FVector2D ScreenSpacePosition = MouseEvent.GetScreenSpacePosition();
+		const FVector2D CursorPosition = Geometry.AbsoluteToLocal(ScreenSpacePosition);
+		const float CursorXRatio = CursorPosition.X / LocalWidth;
+
+		switch (MouseInteractionType)
+		{
+		default:
+			break;
+		case EReceivedInteractionType::MouseButtonDown:
+			if (MouseEvent.IsMouseButtonDown(EKeys::LeftMouseButton))
+			{
+				return FReply::Handled().CaptureMouse(TimeRulerWidget.AsShared()).PreventThrottling();
+			}
+		case EReceivedInteractionType::MouseMove:
+			if (MouseEvent.IsMouseButtonDown(EKeys::LeftMouseButton))
+			{
+				ScrubPlayhead(CursorXRatio, true);
+				return FReply::Handled().CaptureMouse(TimeRulerWidget.AsShared());
+			}
+		case EReceivedInteractionType::MouseButtonUp:
+			if (TimeRulerWidget.HasMouseCapture())
+			{
+				ScrubPlayhead(CursorXRatio, false);
+				return FReply::Handled().ReleaseMouseCapture();
+			}
+		}
+	}
+	
+	return FReply::Handled();
+}
+
+void FWaveformEditorTransportCoordinator::HandleRenderDataUpdate()
+{
+	PlaybackRange = RenderData->GetTransformedWaveformBounds();
+	ReceivePlayBackRatio(0.f);
+}
+
+
+void FWaveformEditorTransportCoordinator::ScrubPlayhead(const float TargetPlayheadPosition, const bool bIsMoving)
+{
+	const float LowerPlaybackBoundPosition = (PlaybackRange.GetLowerBoundValue() - DisplayRange.GetLowerBoundValue()) / ZoomRatio;
+	const float UpperPlaybackBoundPosition = (PlaybackRange.GetUpperBoundValue() - DisplayRange.GetLowerBoundValue()) / ZoomRatio;
+	float ClampedTargetPosition = FMath::Clamp(TargetPlayheadPosition, LowerPlaybackBoundPosition, UpperPlaybackBoundPosition);
+	float NewPosition = FMath::Clamp(ClampedTargetPosition, 0.f, 1.f);
+
+	if (ClampedTargetPosition > PlayheadLockPosition)
+	{
+		const float BoundsDelta = (ClampedTargetPosition - PlayheadLockPosition) * ZoomRatio;
+		const float NewDisplayUpBound = FMath::Clamp(DisplayRange.GetUpperBoundValue() + BoundsDelta, 0.f, 1.f);
+
+		if (NewDisplayUpBound < 1.f)
+		{
+			const float NewDisplayLowBound = FMath::Clamp(DisplayRange.GetLowerBoundValue() + BoundsDelta, 0.f, 1.f);
+			UpdateDisplayRange(NewDisplayLowBound, NewDisplayUpBound);
+			NewPosition = PlayheadLockPosition;
+		}
+	}
+
+	if (ClampedTargetPosition < 0.f)
+	{
+		const float BoundsDelta = ClampedTargetPosition * ZoomRatio;
+		const float NewDisplayLowBound = FMath::Clamp(DisplayRange.GetLowerBoundValue() + BoundsDelta, 0.f, 1.f);
+
+		if (NewDisplayLowBound > 0.f)
+		{
+			const float NewDisplayUpBound = FMath::Clamp(DisplayRange.GetUpperBoundValue() + BoundsDelta, 0.f, 1.f);
+			UpdateDisplayRange(NewDisplayLowBound, NewDisplayUpBound);
+			NewPosition = 0.f;
+		}
+	}
+
+	MovePlayhead(NewPosition);
+
+	if (OnPlayheadScrubUpdate.IsBound())
+	{
+		OnPlayheadScrubUpdate.Broadcast(GetSampleFromPlayheadPosition(NewPosition), RenderData->GetOriginalWaveformFrames(), bIsMoving);
+	}
+
+	bIsScrubbing = bIsMoving;
+}
+const bool FWaveformEditorTransportCoordinator::IsScrubbing() const
+{
+	return bIsScrubbing;
+}
+
+void FWaveformEditorTransportCoordinator::ReceivePlayBackRatio(const float NewRatio)
+{
+	if (bIsScrubbing)
+	{
+		return;
+	}
+		
+	check(NewRatio >= 0.f && NewRatio <= 1.f);
+
+	const float PlaybackRangeScale = PlaybackRange.Size<float>();
+	const float ScaledPlaybackRatio = NewRatio * PlaybackRangeScale + PlaybackRange.GetLowerBoundValue();
+	const float NewPlayheadPosition = (ScaledPlaybackRatio - DisplayRange.GetLowerBoundValue()) / ZoomRatio;
+
+	MovePlayhead(NewPlayheadPosition);
+
+	if (PlayheadPosition < 0.f)
+	{
+		//if sound is looping and we are zoomed, playhead position will be negative when new loop starts
+		//hence we reset the view to go back at the beginning of the file
+		MovePlayhead(0.f);
+		UpdateDisplayRange(PlayheadPosition, ZoomRatio);
+	}
+	else if (PlayheadPosition >= PlayheadLockPosition)
+	{
+		const float NewDisplayUpBound = FMath::Clamp(ScaledPlaybackRatio + (ZoomRatio * (1.f - PlayheadLockPosition)), 0.f, 1.f);
+
+		if (NewDisplayUpBound < 1.f)
+		{
+			//we only lock the playhead if the display upper bound is before the end of the waveform
+			//otherwise we let the playhead go to the end
+			MovePlayhead(PlayheadLockPosition);
+
+			const float NewDisplayLowBound = FMath::Clamp(ScaledPlaybackRatio - (ZoomRatio * PlayheadLockPosition), 0.f, 1.f);
+			UpdateDisplayRange(NewDisplayLowBound, NewDisplayUpBound);
+		}
+	}
+}
+
+void FWaveformEditorTransportCoordinator::OnZoomLevelChanged(const uint8 NewLevel)
+{
+	const float NewZoomRatio = NewLevel / 100.f;
+	UpdateZoomRatioAndDisplayRange(NewZoomRatio);
+}
+
+TRange<float> FWaveformEditorTransportCoordinator::GetDisplayRange() const
+{
+	return DisplayRange;
+}
+
+void FWaveformEditorTransportCoordinator::MovePlayhead(const float InPlayheadPosition)
+{
+	PlayheadPosition = InPlayheadPosition;
+}
+
+void FWaveformEditorTransportCoordinator::UpdateZoomRatioAndDisplayRange(const float NewZoomRatio)
+{
+	const float ClampedZoomRatio = FMath::Clamp(NewZoomRatio, UE_KINDA_SMALL_NUMBER, 1.f);
+	const float ZoomRatioDelta = ZoomRatio - ClampedZoomRatio;
+ 	const float DeltaOutsideBoundaries = DisplayRange.GetUpperBoundValue() - (DisplayRange.GetUpperBoundValue() - ZoomRatioDelta);
+	
+	const float DeltaStepL = DeltaOutsideBoundaries * PlayheadPosition;
+	const float DeltaStepR = DeltaOutsideBoundaries * (1.f - PlayheadPosition);
+
+	float MinDisplayRangeValue = DeltaStepL + DisplayRange.GetLowerBoundValue();
+	float MaxDisplayRangeValue = DisplayRange.GetUpperBoundValue() - DeltaStepR;
+
+	const float CurrentPlaybackPosition = (ZoomRatio * PlayheadPosition) + DisplayRange.GetLowerBoundValue();
+	
+	//if Min < 0 or Max > 1 we can't zoom out further on that side
+	//hence we shift the zooming entirely on the opposite side
+	if (MinDisplayRangeValue < 0.f)
+	{
+		const float ExceedingDelta = FMath::Abs(MinDisplayRangeValue);
+		MaxDisplayRangeValue = MaxDisplayRangeValue + ExceedingDelta;
+		MinDisplayRangeValue = 0.f;
+	}
+	else if (MaxDisplayRangeValue > 1.f)
+	{
+		const float ExceedingDelta = MaxDisplayRangeValue - 1.f;
+		MinDisplayRangeValue = MinDisplayRangeValue - ExceedingDelta;
+		MaxDisplayRangeValue = 1.f;
+	}
+
+	const float NewPlayheadPosition = (CurrentPlaybackPosition - MinDisplayRangeValue) / ClampedZoomRatio;
+	ZoomRatio = ClampedZoomRatio;
+	UpdateDisplayRange(MinDisplayRangeValue, MaxDisplayRangeValue);
+	MovePlayhead(NewPlayheadPosition);
+}
+
+void FWaveformEditorTransportCoordinator::Stop()
+{
+	ReceivePlayBackRatio(0.f);
+	UpdateDisplayRange(0.f, ZoomRatio);
+}
+
+void FWaveformEditorTransportCoordinator::UpdateDisplayRange(const float MinValue, const float MaxValue)
+{
+	check(MinValue < MaxValue);
+
+	DisplayRange.SetLowerBoundValue(MinValue);
+	DisplayRange.SetUpperBoundValue(MaxValue);
+
+	if (OnDisplayRangeUpdated.IsBound())
+	{
+		OnDisplayRangeUpdated.Broadcast(DisplayRange);
+	}
+
+}
+
+bool FWaveformEditorTransportCoordinator::IsRatioWithinDisplayRange(const float Ratio) const
+{
+	return DisplayRange.Contains(Ratio);
+}
+
+uint32 FWaveformEditorTransportCoordinator::GetSampleFromPlayheadPosition(const float InPlayheadPosition) const
+{
+	check(InPlayheadPosition >= 0.f && InPlayheadPosition <= 1.f);
+	const float PlayBackRatio = ((InPlayheadPosition * ZoomRatio) + DisplayRange.GetLowerBoundValue() - PlaybackRange.GetLowerBoundValue()) / PlaybackRange.Size<float>();
+	const uint32 TargetSample = FMath::RoundToInt32(PlayBackRatio * RenderData->GetOriginalWaveformFrames());
+	return TargetSample;
+}
+
+const float FWaveformEditorTransportCoordinator::GetPlayheadPosition() const
+{
+	return PlayheadPosition;
+}
diff --git a/Engine/Plugins/Editor/WaveformEditor/Source/WaveformEditorWidgets/Private/WaveformEditorWidgetsModule.cpp b/Engine/Plugins/Editor/WaveformEditor/Source/WaveformEditorWidgets/Private/WaveformEditorWidgetsModule.cpp
new file mode 100644
--- /dev/null
+++ b/Engine/Plugins/Editor/WaveformEditor/Source/WaveformEditorWidgets/Private/WaveformEditorWidgetsModule.cpp
@@ -0,0 +1,18 @@
+// Copyright Epic Games, Inc. All Rights Reserved.
+
+#include "WaveformEditorWidgetsModule.h"
+
+#include "Modules/ModuleManager.h"
+#include "WaveformEditorStyle.h"
+
+void FWaveformEditorWidgetsModule::StartupModule()
+{
+	// Initialize static style instance 
+	FWaveformEditorStyle::Init();
+}
+
+void FWaveformEditorWidgetsModule::ShutdownModule()
+{
+}
+
+IMPLEMENT_MODULE(FWaveformEditorWidgetsModule, WaveformEditorWidgets)
diff --git a/Engine/Plugins/Editor/WaveformEditor/Source/WaveformEditorWidgets/Private/WaveformEditorWidgetsSettings.cpp b/Engine/Plugins/Editor/WaveformEditor/Source/WaveformEditorWidgets/Private/WaveformEditorWidgetsSettings.cpp
new file mode 100644
--- /dev/null
+++ b/Engine/Plugins/Editor/WaveformEditor/Source/WaveformEditorWidgets/Private/WaveformEditorWidgetsSettings.cpp
@@ -0,0 +1,49 @@
+// Copyright Epic Games, Inc. All Rights Reserved.
+
+#include "WaveformEditorWidgetsSettings.h"
+
+UWaveformEditorWidgetsSettings::UWaveformEditorWidgetsSettings(const FObjectInitializer& ObjectInitializer)
+	: Super(ObjectInitializer)
+	, PlayheadColor(FLinearColor(255, 0.1, 0.2, 1.0))
+	, WaveformColor(FLinearColor::White)
+	, WaveformBackgroundColor(FLinearColor(0.02, 0.02, 0.02, 1.f))
+	, MajorGridColor(FLinearColor::Black)
+	, MinorGridColor(FLinearColor(0.f, 0.f, 0.f, 0.5))
+	, RulerBackgroundColor(FLinearColor::Black)
+	, RulerTicksColor(FLinearColor(1.f, 1.f, 1.f, 0.9))
+	, RulerTextColor(FLinearColor(1.f, 1.f, 1.f, 0.9))
+	, RulerFontSize(10.f)
+{
+}
+
+FName UWaveformEditorWidgetsSettings::GetCategoryName() const
+{
+	return TEXT("Plugins");
+}
+
+FText UWaveformEditorWidgetsSettings::GetSectionText() const
+{
+	return NSLOCTEXT("WaveformEditorDisplay", "WaveformEditorDisplaySettingsSection", "Waveform Editor Display");
+}
+
+FName UWaveformEditorWidgetsSettings::GetSectionName() const
+{
+	return TEXT("Waveform Editor Display");
+}
+
+void UWaveformEditorWidgetsSettings::PostEditChangeProperty(FPropertyChangedEvent& PropertyChangedEvent)
+{
+	Super::PostEditChangeProperty(PropertyChangedEvent);
+
+	if (PropertyChangedEvent.Property != nullptr)
+	{
+		SettingsChangedDelegate.Broadcast(PropertyChangedEvent.GetPropertyName(), this);
+	}
+}
+
+FOnWaveformEditorWidgetsSettingsChanged& UWaveformEditorWidgetsSettings::OnSettingChanged()
+{
+	return SettingsChangedDelegate;
+}
+
+FOnWaveformEditorWidgetsSettingsChanged UWaveformEditorWidgetsSettings::SettingsChangedDelegate;
diff --git a/Engine/Plugins/Editor/WaveformEditor/Source/WaveformEditorWidgets/Private/WaveformEditorWidgetsSettings.h b/Engine/Plugins/Editor/WaveformEditor/Source/WaveformEditorWidgets/Private/WaveformEditorWidgetsSettings.h
new file mode 100644
--- /dev/null
+++ b/Engine/Plugins/Editor/WaveformEditor/Source/WaveformEditorWidgets/Private/WaveformEditorWidgetsSettings.h
@@ -0,0 +1,64 @@
+// Copyright Epic Games, Inc. All Rights Reserved.
+
+#pragma once
+
+#include "Engine/DeveloperSettings.h"
+
+#include "WaveformEditorWidgetsSettings.generated.h"
+
+DECLARE_MULTICAST_DELEGATE_TwoParams(FOnWaveformEditorWidgetsSettingsChanged, const FName& /*Property Name*/, const UWaveformEditorWidgetsSettings*);
+
+UCLASS(config = WaveformEditor, defaultconfig, meta = (DisplayName = "Waveform Editor Display"))
+class UWaveformEditorWidgetsSettings : public UDeveloperSettings
+{
+	GENERATED_BODY()
+
+public:
+
+	UWaveformEditorWidgetsSettings(const FObjectInitializer& ObjectInitializer);
+
+	//~ Begin UDeveloperSettings interface
+	virtual FName GetCategoryName() const;
+
+#if WITH_EDITOR
+	virtual FText GetSectionText() const override;
+
+	virtual FName GetSectionName() const override;
+
+	virtual void PostEditChangeProperty(FPropertyChangedEvent& PropertyChangedEvent) override;
+
+#endif
+	//~ End UDeveloperSettings interface
+
+	static FOnWaveformEditorWidgetsSettingsChanged& OnSettingChanged();
+
+	UPROPERTY(config, EditAnywhere, Category = "Playhead")
+	FLinearColor PlayheadColor;
+
+	UPROPERTY(config, EditAnywhere, Category = "Waveform Viewer")
+	FLinearColor WaveformColor;
+
+	UPROPERTY(config, EditAnywhere, Category = "Waveform Viewer")
+	FLinearColor WaveformBackgroundColor;
+
+	UPROPERTY(config, EditAnywhere, Category = "Waveform Viewer")
+	FLinearColor MajorGridColor;
+	
+	UPROPERTY(config, EditAnywhere, Category = "Waveform Viewer")
+	FLinearColor MinorGridColor;
+
+	UPROPERTY(config, EditAnywhere, Category = "Ruler")
+	FLinearColor RulerBackgroundColor;
+
+	UPROPERTY(config, EditAnywhere, Category = "Ruler")
+	FLinearColor RulerTicksColor;
+
+	UPROPERTY(config, EditAnywhere, Category = "Ruler")
+	FLinearColor RulerTextColor;
+
+	UPROPERTY(config, EditAnywhere, Category = "Ruler", Meta = (ClampMin = "1", ClampMax = "15"))
+	float RulerFontSize;
+
+private: 
+	static FOnWaveformEditorWidgetsSettingsChanged SettingsChangedDelegate;
+};
\ No newline at end of file
diff --git a/Engine/Plugins/Editor/WaveformEditor/Source/WaveformEditorWidgets/Private/WaveformEditorZoomController.cpp b/Engine/Plugins/Editor/WaveformEditor/Source/WaveformEditorWidgets/Private/WaveformEditorZoomController.cpp
new file mode 100644
--- /dev/null
+++ b/Engine/Plugins/Editor/WaveformEditor/Source/WaveformEditorWidgets/Private/WaveformEditorZoomController.cpp
@@ -0,0 +1,56 @@
+// Copyright Epic Games, Inc. All Rights Reserved.
+
+#include "WaveformEditorZoomController.h"
+
+void FWaveformEditorZoomController::ZoomIn()
+{
+	if (CanZoomIn())
+	{
+		ZoomPercentage += ZoomPercentageStep;
+		ApplyZoom();
+	}
+
+}
+
+bool FWaveformEditorZoomController::CanZoomIn() const
+{
+	return ZoomPercentage + ZoomPercentageStep <= 100;
+}
+
+void FWaveformEditorZoomController::ZoomOut()
+{
+	if (CanZoomOut())
+	{
+		ZoomPercentage -= ZoomPercentageStep;
+		ApplyZoom();
+	}
+}
+
+bool FWaveformEditorZoomController::CanZoomOut() const
+{
+	return ZoomPercentage - ZoomPercentageStep >= 0;
+}
+
+
+void FWaveformEditorZoomController::ZoomByDelta(const float Delta)
+{
+	if (Delta >= 0.f)
+	{
+		ZoomIn();
+	}
+	else
+	{
+		ZoomOut();
+	}
+}
+
+uint8 FWaveformEditorZoomController::GetZoomRatio() const
+{
+	return 100 - ZoomPercentage;
+}
+
+void FWaveformEditorZoomController::ApplyZoom()
+{
+	UE_LOG(LogInit, Log, TEXT("Wave Editor Zoom %d%%"), ZoomPercentage);
+	OnZoomRatioChanged.Broadcast(100 - ZoomPercentage);
+}
diff --git a/Engine/Plugins/Editor/WaveformEditor/Source/WaveformEditorWidgets/Public/SWaveformPanel.h b/Engine/Plugins/Editor/WaveformEditor/Source/WaveformEditorWidgets/Public/SWaveformPanel.h
new file mode 100644
--- /dev/null
+++ b/Engine/Plugins/Editor/WaveformEditor/Source/WaveformEditorWidgets/Public/SWaveformPanel.h
@@ -0,0 +1,48 @@
+// Copyright Epic Games, Inc. All Rights Reserved.
+
+#pragma once
+
+#include "HAL/Platform.h"
+#include "Widgets/SCompoundWidget.h"
+
+enum class EWaveformEditorDisplayUnit;
+class FWaveformEditorGridData;
+class FWaveformEditorRenderData;
+class FWaveformEditorTransportCoordinator;
+class FWaveformEditorZoomController;
+class SWaveformEditorTimeRuler;
+class SWaveformTransformationsOverlay;
+class SWaveformViewer;
+class SWaveformViewerOverlay;
+
+class WAVEFORMEDITORWIDGETS_API SWaveformPanel : public SCompoundWidget
+{
+public: 
+
+	SLATE_BEGIN_ARGS(SWaveformPanel) {}
+		SLATE_DEFAULT_SLOT(FArguments, InArgs)
+	SLATE_END_ARGS()
+
+	void Construct(const FArguments& InArgs, TSharedRef<FWaveformEditorRenderData> InRenderData, TSharedRef<FWaveformEditorTransportCoordinator> InTransportCoordinator, TSharedRef<FWaveformEditorZoomController> InZoomManager, TSharedPtr<SWaveformTransformationsOverlay> InWaveformTransformationsOverlay = nullptr);
+	virtual void Tick(const FGeometry& AllottedGeometry, const double InCurrentTime, const float InDeltaTime) override;
+
+private: 
+	void CreateLayout();
+
+	void SetUpWaveformViewer(TSharedRef<FWaveformEditorRenderData> InRenderData, TSharedRef<FWaveformEditorTransportCoordinator> InTransportCoordinator, TSharedRef<FWaveformEditorZoomController> InZoomManager, TSharedRef<FWaveformEditorGridData> InGridData);
+	void SetUpWaveformViewerOverlay(TSharedRef<FWaveformEditorTransportCoordinator> InTransportCoordinator, TSharedRef<FWaveformEditorZoomController> InZoomManager);
+	void SetUpTimeRuler(TSharedRef<FWaveformEditorRenderData> InRenderData, TSharedRef<FWaveformEditorTransportCoordinator> InTransportCoordinator, TSharedRef<FWaveformEditorGridData> InGridData);
+	void SetUpGridData(TSharedRef<FWaveformEditorRenderData> InRenderData, TSharedRef<FWaveformEditorTransportCoordinator> InTransportCoordinator);
+
+	void UpdateDisplayUnit(const EWaveformEditorDisplayUnit InDisplayUnit);
+
+	TSharedPtr<FWaveformEditorGridData> GridData = nullptr;
+	TSharedPtr<SWaveformEditorTimeRuler> TimeRuler = nullptr;
+	TSharedPtr<SWaveformViewer> WaveformViewer = nullptr;
+	TSharedPtr<SWaveformTransformationsOverlay> WaveformTransformationsOverlay = nullptr;
+	TSharedPtr<SWaveformViewerOverlay> WaveformViewerOverlay = nullptr;
+
+	float CachedPixelWidth = 0.f;
+
+	EWaveformEditorDisplayUnit DisplayUnit;
+};
\ No newline at end of file
diff --git a/Engine/Plugins/Editor/WaveformEditor/Source/WaveformEditorWidgets/Public/SWaveformTransformationsOverlay.h b/Engine/Plugins/Editor/WaveformEditor/Source/WaveformEditorWidgets/Public/SWaveformTransformationsOverlay.h
new file mode 100644
--- /dev/null
+++ b/Engine/Plugins/Editor/WaveformEditor/Source/WaveformEditorWidgets/Public/SWaveformTransformationsOverlay.h
@@ -0,0 +1,25 @@
+// Copyright Epic Games, Inc. All Rights Reserved.
+
+#pragma once
+
+#include "Widgets/SCompoundWidget.h"
+
+class SOverlay;
+
+class WAVEFORMEDITORWIDGETS_API SWaveformTransformationsOverlay : public SCompoundWidget
+{
+public:
+
+	SLATE_BEGIN_ARGS(SWaveformTransformationsOverlay) {}
+		SLATE_DEFAULT_SLOT(FArguments, InArgs)
+	SLATE_END_ARGS()
+
+	void Construct(const FArguments& InArgs, TArrayView<TSharedPtr<SWidget>> InTransformationLayers);
+	void OnLayerChainUpdate(TSharedPtr<SWidget>* FirstLayerPtr, const int32 NLayers);
+
+private:
+	void CreateLayout();
+	
+	TSharedPtr<SOverlay> MainOverlayPtr;
+	TArrayView<TSharedPtr<SWidget>> TransformationLayers;
+};
\ No newline at end of file
diff --git a/Engine/Plugins/Editor/WaveformEditor/Source/WaveformEditorWidgets/Public/WaveformDrawingUtils.h b/Engine/Plugins/Editor/WaveformEditor/Source/WaveformEditorWidgets/Public/WaveformDrawingUtils.h
new file mode 100644
--- /dev/null
+++ b/Engine/Plugins/Editor/WaveformEditor/Source/WaveformEditorWidgets/Public/WaveformDrawingUtils.h
@@ -0,0 +1,12 @@
+// Copyright Epic Games, Inc. All Rights Reserved.
+
+#pragma once
+
+#include "Math/Range.h"
+
+namespace WaveformDrawingUtils
+{
+	typedef TRange<int16> SampleRange;
+
+	void WAVEFORMEDITORWIDGETS_API GetBinnedPeaksFromWaveformRawData(TArray<SampleRange>& OutWaveformPeaks, const uint32 NBins, const int16* RawDataPtr, const uint32 NSamples, const uint32 SampleRate, const uint16 NChannels, const float StartTime, const float EndTime);
+}
\ No newline at end of file
diff --git a/Engine/Plugins/Editor/WaveformEditor/Source/WaveformEditorWidgets/Public/WaveformEditorGridData.h b/Engine/Plugins/Editor/WaveformEditor/Source/WaveformEditorWidgets/Public/WaveformEditorGridData.h
new file mode 100644
--- /dev/null
+++ b/Engine/Plugins/Editor/WaveformEditor/Source/WaveformEditorWidgets/Public/WaveformEditorGridData.h
@@ -0,0 +1,42 @@
+// Copyright Epic Games, Inc. All Rights Reserved.
+
+#pragma once
+
+#include "Delegates/Delegate.h"
+#include "HAL/Platform.h"
+#include "Math/Range.h"
+
+struct FSlateFontInfo;
+class FWaveformEditorRenderData;
+
+struct WAVEFORMEDITORWIDGETS_API FWaveEditorGridMetrics
+{
+	int32 NumMinorGridDivisions = 0;
+	double PixelsPerSecond = 0;
+	double FirstMajorTickX = 0;
+	double MajorGridXStep = 0;
+	double StartTime = 0;
+};
+
+DECLARE_MULTICAST_DELEGATE_OneParam(FOnGridMetricsUpdated, const FWaveEditorGridMetrics& /* New Grid Metrics */);
+
+class WAVEFORMEDITORWIDGETS_API FWaveformEditorGridData
+{
+public: 
+	explicit FWaveformEditorGridData(TSharedRef<FWaveformEditorRenderData> InRenderData, const FSlateFontInfo* InTicksTimeFont = nullptr);
+
+	FOnGridMetricsUpdated OnGridMetricsUpdated;
+
+	void UpdateDisplayRange(const TRange<float> InDisplayRange);
+	bool UpdateGridMetrics(const float InGridPixelWidth);
+	const FWaveEditorGridMetrics GetGridMetrics() const;
+	void SetTicksTimeFont(const FSlateFontInfo* InNewFont);
+
+private:
+	FWaveEditorGridMetrics GridMetrics;
+	TSharedPtr<FWaveformEditorRenderData> RenderData = nullptr;
+	TRange<float> DisplayRange = TRange<float>::Inclusive(0.f, 1.f);
+
+	float GridPixelWidth = 0.f;
+	const FSlateFontInfo* TicksTimeFont = nullptr;
+};
\ No newline at end of file
diff --git a/Engine/Plugins/Editor/WaveformEditor/Source/WaveformEditorWidgets/Public/WaveformEditorRenderData.h b/Engine/Plugins/Editor/WaveformEditor/Source/WaveformEditorWidgets/Public/WaveformEditorRenderData.h
new file mode 100644
--- /dev/null
+++ b/Engine/Plugins/Editor/WaveformEditor/Source/WaveformEditorWidgets/Public/WaveformEditorRenderData.h
@@ -0,0 +1,58 @@
+// Copyright Epic Games, Inc. All Rights Reserved.
+
+#pragma once
+
+#include "Containers/ArrayView.h"
+#include "Delegates/Delegate.h"
+#include "HAL/Platform.h"
+#include "Math/Range.h"
+
+class USoundWave;
+
+DECLARE_MULTICAST_DELEGATE(FOnRenderDataUpdated)
+
+
+/************************************************************************************************/
+/* FWaveformEditorRenderData																	*/
+/* The Waveform Editor Render Data generates a view of Raw PCM data through the 				*/
+/* UpdateRenderData() method.																	*/
+/* 																								*/
+/* This class can be used to propagate efficiently waveform data to different renderers 		*/
+/* while keeping it separate from them, but please not that it DOES NOT hold onto that data.	*/
+/*																								*/
+/* Users of this class should make sure that the data being shown by the view is valid. 		*/
+/************************************************************************************************/
+
+class WAVEFORMEDITORWIDGETS_API FWaveformEditorRenderData : public TSharedFromThis<FWaveformEditorRenderData>
+{
+public:
+	FWaveformEditorRenderData() = default;
+
+	TArrayView<const int16> GetSampleData() const;
+	const uint16 GetNumChannels() const;
+	const uint32 GetNumSamples() const;
+	const uint32 GetSampleRate() const;
+	const float GetOriginalWaveformDurationInSeconds() const;
+	const uint32 GetOriginalWaveformFrames() const;
+	const TRange<float> GetTransformedWaveformBounds() const;
+
+	/** Call this to generate the view and the info shown by the render data */
+	void UpdateRenderData(const uint8* InRawData, const uint32 InNumSamples, const uint32 InFirstEditedSample, const uint32 LastEditedSample, const uint32 InSampleRate, const uint16 InNumChannels);
+
+	/** Called when the render data view is updated */
+	FOnRenderDataUpdated OnRenderDataUpdated;
+	
+private:
+	TArrayView<const int16> SampleData;
+
+	uint16 NumChannels = 0;
+	uint32 NumSamples = 0;
+	uint32 NumEditedSamples = 0;
+	uint32 SampleRate = 0;
+
+	float OriginalWaveformDurationInSeconds = 0.f;
+	uint32 OriginalWaveformFrames = 0;
+
+	/* The bounds of the transformed waveform in relation to the original */
+	TRange<float> TransformedWaveformBounds = TRange<float>::Inclusive(0.f, 1.f);
+};
\ No newline at end of file
diff --git a/Engine/Plugins/Editor/WaveformEditor/Source/WaveformEditorWidgets/Public/WaveformEditorSlateTypes.h b/Engine/Plugins/Editor/WaveformEditor/Source/WaveformEditorWidgets/Public/WaveformEditorSlateTypes.h
new file mode 100644
--- /dev/null
+++ b/Engine/Plugins/Editor/WaveformEditor/Source/WaveformEditorWidgets/Public/WaveformEditorSlateTypes.h
@@ -0,0 +1,148 @@
+// Copyright Epic Games, Inc. All Rights Reserved.
+
+#pragma once
+
+
+#include "Styling/SlateStyle.h"
+#include "Styling/SlateWidgetStyle.h"
+
+#include "WaveformEditorSlateTypes.generated.h"
+
+/**
+ * Represents the appearance of a Waveform Viewer
+ */
+USTRUCT(BlueprintType)
+struct FWaveformViewerStyle : public FSlateWidgetStyle
+{
+	GENERATED_USTRUCT_BODY()
+
+	FWaveformViewerStyle();
+
+	static const FWaveformViewerStyle& GetDefault();
+	virtual const FName GetTypeName() const override { return TypeName; };
+	virtual void GetResources(TArray< const FSlateBrush* >& OutBrushes) const override;
+
+	static const FName TypeName;
+
+	UPROPERTY(EditAnywhere, BlueprintReadWrite, Category = Appearance)
+	FSlateColor WaveformColor;
+	FWaveformViewerStyle& SetWaveformColor(const FSlateColor InWaveformColor) { WaveformColor = InWaveformColor; return *this; }
+
+	UPROPERTY(EditAnywhere, BlueprintReadWrite, Category = Appearance)
+	FSlateColor MajorGridLineColor;
+	FWaveformViewerStyle& SetMajorGridLineColor(const FSlateColor InMajorGridLineColor) { MajorGridLineColor = InMajorGridLineColor; return *this; }
+
+	UPROPERTY(EditAnywhere, BlueprintReadWrite, Category = Appearance)
+	FSlateColor MinorGridLineColor;
+	FWaveformViewerStyle& SetMinorGridLineColor(const FSlateColor InMinorGridLineColor) { MinorGridLineColor = InMinorGridLineColor; return *this; }
+
+	UPROPERTY(EditAnywhere, BlueprintReadWrite, Category = Appearance)
+	FSlateColor WaveformBackgroundColor;
+	FWaveformViewerStyle& SetBackgroundColor(const FSlateColor InBackgroundColor) { WaveformBackgroundColor = InBackgroundColor; return *this; }
+
+	UPROPERTY(EditAnywhere, BlueprintReadWrite, Category = Appearance)
+	FSlateBrush BackgroundBrush;
+	FWaveformViewerStyle& SetBackgroundBrush(const FSlateBrush InBackgroundBrush) { BackgroundBrush = InBackgroundBrush; return *this; }
+	
+	UPROPERTY(EditAnywhere, BlueprintReadWrite, Category = Appearance)
+	float DesiredWidth;
+	FWaveformViewerStyle& SetDesiredWidth(const float InDesiredWidth) { DesiredWidth = InDesiredWidth; return *this; }
+
+	UPROPERTY(EditAnywhere, BlueprintReadWrite, Category = Appearance)
+	float DesiredHeight;
+	FWaveformViewerStyle& SetDesiredHeight(const float InDesiredHeight) { DesiredHeight = InDesiredHeight; return *this; }
+};
+
+/**
+ * Represents the appearance of a Waveform Viewer Overlay style
+ */
+USTRUCT(BlueprintType)
+struct FWaveformViewerOverlayStyle : public FSlateWidgetStyle
+{
+	GENERATED_USTRUCT_BODY()
+
+	FWaveformViewerOverlayStyle();
+
+	static const FWaveformViewerOverlayStyle& GetDefault();
+	virtual const FName GetTypeName() const override { return TypeName; };
+
+	static const FName TypeName;
+
+	UPROPERTY(EditAnywhere, BlueprintReadWrite, Category = Appearance)
+	FSlateColor PlayheadColor;
+	FWaveformViewerOverlayStyle& SetPlayheadColor(const FSlateColor InPlayheadColor) { PlayheadColor = InPlayheadColor; return *this; }
+
+	UPROPERTY(EditAnywhere, BlueprintReadWrite, Category = Appearance)
+	float PlayheadWidth;
+	FWaveformViewerOverlayStyle& SetPlayheadWidth(const float InPlayheadWidth) { PlayheadWidth = InPlayheadWidth; return *this; }
+
+	UPROPERTY(EditAnywhere, BlueprintReadWrite, Category = Appearance)
+	float DesiredWidth;
+	FWaveformViewerOverlayStyle& SetDesiredWidth(const float InDesiredWidth) { DesiredWidth = InDesiredWidth; return *this; }
+
+	UPROPERTY(EditAnywhere, BlueprintReadWrite, Category = Appearance)
+	float DesiredHeight;
+	FWaveformViewerOverlayStyle& SetDesiredHeight(const float InDesiredHeight) { DesiredHeight = InDesiredHeight; return *this; }
+};
+
+/**
+ * Represents the appearance of a Waveform Editor Time Ruler
+ */
+USTRUCT(BlueprintType)
+struct FWaveformEditorTimeRulerStyle : public FSlateWidgetStyle
+{
+	GENERATED_USTRUCT_BODY()
+
+	FWaveformEditorTimeRulerStyle();
+
+	static const FWaveformEditorTimeRulerStyle& GetDefault();
+	virtual const FName GetTypeName() const override { return TypeName; };
+	virtual void GetResources(TArray< const FSlateBrush* >& OutBrushes) const override;
+
+	static const FName TypeName;
+
+	UPROPERTY(EditAnywhere, BlueprintReadWrite, Category = Appearance)
+	float HandleWidth;
+	FWaveformEditorTimeRulerStyle& SetHandleWidth(const float InHandleWidth) { HandleWidth = InHandleWidth; return *this; }
+
+	UPROPERTY(EditAnywhere, BlueprintReadWrite, Category = Appearance)
+	FSlateColor HandleColor;
+	FWaveformEditorTimeRulerStyle& SetHandleColor(const FSlateColor& InHandleColor) { HandleColor = InHandleColor; return *this; }
+
+	UPROPERTY(EditAnywhere, BlueprintReadWrite, Category = Appearance)
+	FSlateBrush HandleBrush;
+	FWaveformEditorTimeRulerStyle& SetHandleBrush(const FSlateBrush& InHandleBrush) { HandleBrush = InHandleBrush; return *this; }
+
+	UPROPERTY(EditAnywhere, BlueprintReadWrite, Category = Appearance)
+	FSlateColor TicksColor;
+	FWaveformEditorTimeRulerStyle& SetTicksColor(const FSlateColor& InTicksColor) { TicksColor = InTicksColor; return *this; }
+
+	UPROPERTY(EditAnywhere, BlueprintReadWrite, Category = Appearance)
+	FSlateColor TicksTextColor;
+	FWaveformEditorTimeRulerStyle& SetTicksTextColor(const FSlateColor& InTicksTextColor) { TicksTextColor = InTicksTextColor; return *this; }
+
+	UPROPERTY(EditAnywhere, BlueprintReadWrite, Category = Appearance)
+	FSlateFontInfo TicksTextFont;
+	FWaveformEditorTimeRulerStyle& SetTicksTextFont(const FSlateFontInfo& InTicksTextFont) { TicksTextFont = InTicksTextFont; return *this; }
+	FWaveformEditorTimeRulerStyle& SetFontSize(const float InFontSize) {TicksTextFont.Size = InFontSize; return *this;	}
+
+	UPROPERTY(EditAnywhere, BlueprintReadWrite, Category = Appearance)
+	float TicksTextOffset;
+	FWaveformEditorTimeRulerStyle& SetTicksTextOffset(const float InTicksTextOffset) { TicksTextOffset = InTicksTextOffset; return *this; }
+
+	UPROPERTY(EditAnywhere, BlueprintReadWrite, Category = Appearance)
+	FSlateColor BackgroundColor;
+	FWaveformEditorTimeRulerStyle& SetBackgroundColor(const FSlateColor& InBackgroundColor) { BackgroundColor = InBackgroundColor; return *this; }
+
+	UPROPERTY(EditAnywhere, BlueprintReadWrite, Category = Appearance)
+	FSlateBrush BackgroundBrush;
+	FWaveformEditorTimeRulerStyle& SetBackgroundBrush(const FSlateBrush& InBackgroundBrush) { BackgroundBrush = InBackgroundBrush; return *this; }
+
+	UPROPERTY(EditAnywhere, BlueprintReadWrite, Category = Appearance)
+	float DesiredWidth;
+	FWaveformEditorTimeRulerStyle& SetDesiredWidth(const float InDesiredWidth) { DesiredWidth = InDesiredWidth; return *this; }
+
+	UPROPERTY(EditAnywhere, BlueprintReadWrite, Category = Appearance)
+	float DesiredHeight;
+	FWaveformEditorTimeRulerStyle& SetDesiredHeight(const float InDesiredHeight) { DesiredHeight = InDesiredHeight; return *this; }
+};
\ No newline at end of file
diff --git a/Engine/Plugins/Editor/WaveformEditor/Source/WaveformEditorWidgets/Public/WaveformEditorStyle.h b/Engine/Plugins/Editor/WaveformEditor/Source/WaveformEditorWidgets/Public/WaveformEditorStyle.h
new file mode 100644
--- /dev/null
+++ b/Engine/Plugins/Editor/WaveformEditor/Source/WaveformEditorWidgets/Public/WaveformEditorStyle.h
@@ -0,0 +1,38 @@
+// Copyright Epic Games, Inc. All Rights Reserved.
+
+#pragma once
+
+#include "Styling/SlateStyle.h"
+#include "WaveformEditorSlateTypes.h"
+
+class UWaveformEditorWidgetsSettings;
+
+class FWaveformEditorStyle
+	: public FSlateStyleSet
+
+{
+public:
+	FWaveformEditorStyle();
+	~FWaveformEditorStyle();
+
+	static FName StyleName;
+
+	/** Access the singleton instance for this style set */
+	static const FWaveformEditorStyle& Get();
+	static void Init();
+	
+private:
+	static TUniquePtr<FWaveformEditorStyle> StyleInstance;
+
+	static const UWaveformEditorWidgetsSettings* GetWidgetsSettings();
+	static void  OnWidgetSettingsUpdated(const FName& PropertyName, const UWaveformEditorWidgetsSettings* Settings);
+
+	template< typename SlateWidgetStyle >
+	static TSharedRef<SlateWidgetStyle> GetRegisteredWidgetStyle(const FName& InStyleName)
+	{
+		TSharedRef<FSlateWidgetStyle>* WidgetStyleBase = StyleInstance->WidgetStyleValues.Find(InStyleName);
+		check(WidgetStyleBase);
+		TSharedRef<SlateWidgetStyle> WidgetStyle = StaticCastSharedRef<SlateWidgetStyle>(*WidgetStyleBase);
+		return WidgetStyle;
+	}
+};
\ No newline at end of file
diff --git a/Engine/Plugins/Editor/WaveformEditor/Source/WaveformEditorWidgets/Public/WaveformEditorTransportCoordinator.h b/Engine/Plugins/Editor/WaveformEditor/Source/WaveformEditorWidgets/Public/WaveformEditorTransportCoordinator.h
new file mode 100644
--- /dev/null
+++ b/Engine/Plugins/Editor/WaveformEditor/Source/WaveformEditorWidgets/Public/WaveformEditorTransportCoordinator.h
@@ -0,0 +1,75 @@
+// Copyright Epic Games, Inc. All Rights Reserved.
+
+#pragma once
+
+#include "Widgets/SWidget.h"
+
+DECLARE_MULTICAST_DELEGATE_OneParam(FOnDisplayRangeUpdated, const TRange<float> /* New Display Range */);
+DECLARE_MULTICAST_DELEGATE_ThreeParams(FOnPlayheadScrubUpdate, const uint32 /* New Playhead Position in Samples */, const uint32 /* Waveform Length In Samples */, const bool /*Playhead is Moving*/);
+
+class FWaveformEditorRenderData;
+
+enum class EReceivedInteractionType
+{
+	MouseButtonUp, 
+	MouseButtonDown, 
+	MouseMove
+};
+
+class WAVEFORMEDITORWIDGETS_API FWaveformEditorTransportCoordinator : public TSharedFromThis<FWaveformEditorTransportCoordinator>
+{
+public:
+	explicit FWaveformEditorTransportCoordinator(TSharedRef<FWaveformEditorRenderData> InRenderData);
+	virtual ~FWaveformEditorTransportCoordinator() = default;
+
+	/** Called when the playhead is scrubbed */
+	FOnPlayheadScrubUpdate OnPlayheadScrubUpdate;
+
+	/** Called when the display range is updated */
+	FOnDisplayRangeUpdated OnDisplayRangeUpdated;
+
+	/* Intercept widget interactions */
+	FReply ReceiveMouseButtonDown(SWidget& WidgetOwner, const FGeometry& Geometry, const FPointerEvent& MouseEvent);
+	FReply ReceiveMouseButtonUp(SWidget& WidgetOwner, const FGeometry& Geometry, const FPointerEvent& MouseEvent);
+	FReply ReceiveMouseMove(SWidget& WidgetOwner, const FGeometry& Geometry, const FPointerEvent& MouseEvent);
+
+	void ScrubPlayhead(const float InPlayHeadPosition, const bool bIsMoving);
+	const bool IsScrubbing() const;
+
+	void HandleRenderDataUpdate();
+
+	const float GetPlayheadPosition() const;
+	void ReceivePlayBackRatio(const float NewRatio);
+	void OnZoomLevelChanged(const uint8 NewLevel);
+
+	TRange<float> GetDisplayRange() const;
+
+	void Stop();
+
+private:
+	FORCEINLINE void MovePlayhead(const float InPlayheadPosition);
+	void UpdateZoomRatioAndDisplayRange(const float NewZoomRatio);
+
+	void UpdateDisplayRange(const float MinValue, const float MaxValue);
+	bool IsRatioWithinDisplayRange(const float Ratio) const;
+
+	uint32 GetSampleFromPlayheadPosition(const float InPlayheadPosition) const;
+
+	FReply HandleWaveformViewerOverlayInteraction(const FPointerEvent& MouseEvent, const FGeometry& Geometry);
+	FReply HandleTimeRulerInteraction(const EReceivedInteractionType MouseInteractionType, const FPointerEvent& MouseEvent, SWidget& TimeRulerWidget, const FGeometry& Geometry);
+	
+	float CurrentPlaybackRatio = 0.f;
+	float PlayheadLockPosition = 0.95f;
+	float PlayheadPosition = 0.f;
+	float ZoomRatio = 1.f;
+	
+	/* The currently displayed render data range */
+	TRange<float> DisplayRange = TRange<float>::Inclusive(0.f, 1.f);
+
+	/* Playback range to scale the incoming ratio from the audiocomponent with*/
+	TRange<float> PlaybackRange = TRange<float>::Inclusive(0.f, 1.f);
+	
+	TSharedPtr<FWaveformEditorRenderData> RenderData = nullptr;
+
+	bool bIsScrubbing = false;
+};
\ No newline at end of file
diff --git a/Engine/Plugins/Editor/WaveformEditor/Source/WaveformEditorWidgets/Public/WaveformEditorWidgetsModule.h b/Engine/Plugins/Editor/WaveformEditor/Source/WaveformEditorWidgets/Public/WaveformEditorWidgetsModule.h
new file mode 100644
--- /dev/null
+++ b/Engine/Plugins/Editor/WaveformEditor/Source/WaveformEditorWidgets/Public/WaveformEditorWidgetsModule.h
@@ -0,0 +1,13 @@
+// Copyright Epic Games, Inc. All Rights Reserved.
+
+#pragma once
+
+#include "Modules/ModuleInterface.h"
+
+class FWaveformEditorWidgetsModule : public IModuleInterface
+{
+public: 
+	/** IModuleInterface implementation */
+	virtual void StartupModule() override;
+	virtual void ShutdownModule() override;
+};
diff --git a/Engine/Plugins/Editor/WaveformEditor/Source/WaveformEditorWidgets/Public/WaveformEditorZoomController.h b/Engine/Plugins/Editor/WaveformEditor/Source/WaveformEditorWidgets/Public/WaveformEditorZoomController.h
new file mode 100644
--- /dev/null
+++ b/Engine/Plugins/Editor/WaveformEditor/Source/WaveformEditorWidgets/Public/WaveformEditorZoomController.h
@@ -0,0 +1,29 @@
+// Copyright Epic Games, Inc. All Rights Reserved.
+
+#pragma once
+
+#include "Delegates/Delegate.h"
+
+DECLARE_MULTICAST_DELEGATE_OneParam(FOnZoomRatioChanged, const uint8 /* New Zoom Ratio */);
+
+class WAVEFORMEDITORWIDGETS_API FWaveformEditorZoomController
+{
+public:
+	FWaveformEditorZoomController() = default;
+
+	void ZoomIn();
+	bool CanZoomIn() const;
+	void ZoomOut();
+	bool CanZoomOut() const;
+	void ZoomByDelta(const float Delta);
+
+	uint8 GetZoomRatio() const;
+
+	FOnZoomRatioChanged OnZoomRatioChanged;
+
+private:
+	void ApplyZoom();
+
+	uint8 ZoomPercentage = 0;
+	uint8 ZoomPercentageStep = 5;
+};
\ No newline at end of file
diff --git a/Engine/Plugins/Editor/WaveformEditor/Source/WaveformEditorWidgets/WaveformEditorWidgets.build.cs b/Engine/Plugins/Editor/WaveformEditor/Source/WaveformEditorWidgets/WaveformEditorWidgets.build.cs
new file mode 100644
--- /dev/null
+++ b/Engine/Plugins/Editor/WaveformEditor/Source/WaveformEditorWidgets/WaveformEditorWidgets.build.cs
@@ -0,0 +1,32 @@
+// Copyright Epic Games, Inc. All Rights Reserved.
+
+using UnrealBuildTool;
+using System.IO;
+
+namespace UnrealBuildTool.Rules
+{
+	public class WaveformEditorWidgets : ModuleRules
+	{
+		public WaveformEditorWidgets(ReadOnlyTargetRules Target) : base(Target)
+		{
+			PCHUsage = PCHUsageMode.UseExplicitOrSharedPCHs;
+
+			PublicDependencyModuleNames.AddRange(
+				new string[] {
+					"AudioExtensions",
+					"ContentBrowser",
+					"Core",
+					"CoreUObject",
+					"DeveloperSettings",
+					"EditorStyle",
+					"Engine",
+					"InputCore",
+					"Slate",
+					"SignalProcessing",
+					"SlateCore", 
+					"ToolMenus"
+				}
+			);
+		}
+	}
+}
\ No newline at end of file
diff --git a/Engine/Plugins/Editor/WaveformEditor/Source/WaveformTransformations/Private/WaveformTransformationEffectChain.cpp b/Engine/Plugins/Editor/WaveformEditor/Source/WaveformTransformations/Private/WaveformTransformationEffectChain.cpp
new file mode 100644
--- /dev/null
+++ b/Engine/Plugins/Editor/WaveformEditor/Source/WaveformTransformations/Private/WaveformTransformationEffectChain.cpp
@@ -0,0 +1,105 @@
+// Copyright Epic Games, Inc. All Rights Reserved.
+
+#include "WaveformTransformationEffectChain.h"
+#include "Engine/Engine.h"
+#include "Sound/SoundEffectPreset.h"
+#include "Sound/SoundEffectSource.h"
+
+FWaveTransformationEffectChain::FWaveTransformationEffectChain(TArray<TObjectPtr<USoundEffectSourcePreset>>& InEffectPresets)
+{
+	for(USoundEffectSourcePreset* Preset : InEffectPresets)
+	{
+		if(Preset)
+		{
+			Presets.Add(TStrongObjectPtr<USoundEffectSourcePreset>(Preset));
+		}
+	}
+}
+
+void FWaveTransformationEffectChain::ProcessAudio(Audio::FWaveformTransformationWaveInfo& InOutWaveInfo) const
+{
+	check(InOutWaveInfo.Audio != nullptr);
+	check(InOutWaveInfo.Audio->Num() % InOutWaveInfo.NumChannels == 0);
+
+	uint32 AudioDeviceId = GEngine->GetMainAudioDeviceID();
+
+	FSoundEffectSourceInitData InitData;
+	InitData.SampleRate = InOutWaveInfo.SampleRate;
+	InitData.AudioClock = 0.0;
+	InitData.NumSourceChannels = InOutWaveInfo.NumChannels;
+	InitData.AudioDeviceId = AudioDeviceId;
+
+	TArray<TSharedPtr<class FSoundEffectSource>> Effects;
+
+	for(const TStrongObjectPtr<USoundEffectSourcePreset>& Preset : Presets)
+	{
+		InitData.ParentPresetUniqueId = Preset->GetUniqueID();
+
+		TSoundEffectSourcePtr Effect = USoundEffectPreset::CreateInstance<FSoundEffectSourceInitData, FSoundEffectSource>(InitData, *Preset);
+		Effect->SetEnabled(true);
+		Effects.Add(Effect);
+	}
+	
+	Audio::FAlignedFloatBuffer& InputAudio = *InOutWaveInfo.Audio;
+
+	const int32 NumBlockSamples = 1024 * InOutWaveInfo.NumChannels;
+	
+	Audio::FAlignedFloatBuffer OutputBuffer;
+
+	OutputBuffer.AddUninitialized(NumBlockSamples);
+
+	FSoundEffectSourceInputData InputData;
+	InputData.AudioClock = 0.0;
+	InputData.CurrentPitch = 1.f;
+	InputData.CurrentVolume = 1.f;
+	InputData.NumSamples = NumBlockSamples;
+	InputData.SpatParams = FSpatializationParams();
+	InputData.CurrentPlayFraction = 0.f;
+	InputData.InputSourceEffectBufferPtr = InputAudio.GetData();
+
+	for(TSoundEffectSourcePtr& EffectPtr : Effects)
+	{
+		EffectPtr->Update();
+		
+		int32 SampleIndex = 0;
+		int32 SamplesToProcess = 0;
+
+		while(SampleIndex < InputAudio.Num())
+		{
+			InputData.InputSourceEffectBufferPtr = &InputAudio[SampleIndex];
+			SamplesToProcess = FMath::Clamp(InputAudio.Num() - SampleIndex, 0, NumBlockSamples);
+
+			EffectPtr->ProcessAudio(InputData, OutputBuffer.GetData());
+			
+			FMemory::Memcpy(InputData.InputSourceEffectBufferPtr, OutputBuffer.GetData(), SamplesToProcess * sizeof(float));
+			SampleIndex += SamplesToProcess;
+		}
+	}
+}
+
+Audio::FTransformationPtr UWaveformTransformationEffectChain::CreateTransformation() const
+{
+	TArray<TObjectPtr<USoundEffectSourcePreset>> Effects;
+	if(EffectChain)
+	{
+		for(FSourceEffectChainEntry Entry : EffectChain->Chain)
+		{
+			if(Entry.bBypass == false && Entry.Preset != nullptr)
+			{
+				Effects.Add(Entry.Preset);
+			}
+		}
+	}
+	else
+	{
+		for(USoundEffectSourcePreset* Preset : InlineEffects)
+		{
+			if(Preset)
+			{
+				Effects.Add(Preset);
+			}
+		}
+	}
+	
+	return MakeUnique<FWaveTransformationEffectChain>(Effects);
+}
\ No newline at end of file
diff --git a/Engine/Plugins/Editor/WaveformEditor/Source/WaveformTransformations/Private/WaveformTransformationNormalize.cpp b/Engine/Plugins/Editor/WaveformEditor/Source/WaveformTransformations/Private/WaveformTransformationNormalize.cpp
new file mode 100644
--- /dev/null
+++ b/Engine/Plugins/Editor/WaveformEditor/Source/WaveformTransformations/Private/WaveformTransformationNormalize.cpp
@@ -0,0 +1,127 @@
+// Copyright Epic Games, Inc. All Rights Reserved.
+
+#include "WaveformTransformationNormalize.h"
+
+#include "DSP/FloatArrayMath.h"
+#include "AudioAnalyzer.h"
+#include "IAudioAnalyzerInterface.h"
+#include "LoudnessFactory.h"
+#include "Algo/MaxElement.h"
+#include "MeterFactory.h"
+
+static float GetRMSPeak(const Audio::FAlignedFloatBuffer& InputAudio, float SampleRate, int32 NumChannels)
+{
+	Audio::FMeterFactory Analyzer;
+	Audio::FMeterSettings Settings;
+	Settings.AnalysisPeriod = 0.3f;
+
+	TUniquePtr<Audio::IAnalyzerResult> Result = Analyzer.NewResult();
+	TUniquePtr<Audio::IAnalyzerWorker> AnalyzerWorker = Analyzer.NewWorker({(int32)SampleRate, NumChannels}, &Settings);
+
+	if(AnalyzerWorker == nullptr)
+	{
+		return 0.f;
+	}
+	
+	AnalyzerWorker->Analyze(MakeArrayView(InputAudio), Result.Get());
+
+	Audio::FMeterResult* MeterResult = static_cast<Audio::FMeterResult*>(Result.Get());
+
+	if (MeterResult)
+	{
+		const TArray<Audio::FMeterEntry>& LoudnessArray = MeterResult->GetMeterArray();
+		auto MeterEntry = Algo::MaxElement(LoudnessArray, [](const Audio::FMeterEntry& A, const Audio::FMeterEntry& B)
+		{
+			return A.MeterValue < B.MeterValue;
+		});
+
+		if (MeterEntry)
+		{
+			return MeterEntry->MeterValue;
+		}
+	}
+
+	return 0.f;
+}
+
+static float GetLoudnessPeak(const Audio::FAlignedFloatBuffer& InputAudio, float SampleRate, int32 NumChannels)
+{
+	Audio::FLoudnessFactory Analyzer;
+	Audio::FLoudnessSettings Settings;
+	Settings.AnalysisPeriod = 1.f;
+	
+	TUniquePtr<Audio::IAnalyzerResult> Result = Analyzer.NewResult();
+	TUniquePtr<Audio::IAnalyzerWorker> AnalyzerWorker = Analyzer.NewWorker({(int32)SampleRate, NumChannels}, &Settings);
+
+	if(AnalyzerWorker == nullptr)
+	{
+		return 0.f;
+	}
+	
+	AnalyzerWorker->Analyze(MakeArrayView(InputAudio), Result.Get());
+
+	Audio::FLoudnessResult* LoudnessResult = static_cast<Audio::FLoudnessResult*>(Result.Get());
+
+	if (LoudnessResult)
+	{
+		const TArray<Audio::FLoudnessEntry>& LoudnessArray = LoudnessResult->GetLoudnessArray();
+		auto LoudestEntry = Algo::MaxElement(LoudnessArray, [](const Audio::FLoudnessEntry& A, const Audio::FLoudnessEntry& B)
+		{
+			return A.Loudness < B.Loudness;
+		});
+
+		if (LoudestEntry)
+		{
+			return LoudestEntry->Loudness;
+		}
+	}
+
+	return 0.f;
+}
+
+FWaveTransformationNormalize::FWaveTransformationNormalize(float InTarget, float InMaxGain, ENormalizationMode InMode)
+	: Target(InTarget)
+	, MaxGain(InMaxGain)
+	, Mode(InMode) {}
+
+
+ void FWaveTransformationNormalize::ProcessAudio(Audio::FWaveformTransformationWaveInfo& InOutWaveInfo) const
+{
+	check(InOutWaveInfo.Audio != nullptr);
+	
+	Audio::FAlignedFloatBuffer& InputAudio = *InOutWaveInfo.Audio;
+
+	float PeakDecibelValue = 1.f;
+	
+	switch (Mode)
+	{
+	case ENormalizationMode::Peak:
+		{
+			const float PeakLinearValue = Audio::ArrayMaxAbsValue(InputAudio);
+			PeakDecibelValue = Audio::ConvertToDecibels(PeakLinearValue);
+			break;
+		}
+	case ENormalizationMode::RMS:
+		PeakDecibelValue = GetRMSPeak(InputAudio, InOutWaveInfo.SampleRate, InOutWaveInfo.NumChannels);
+		break;
+	case ENormalizationMode::DWeightedLoudness:
+		PeakDecibelValue = GetLoudnessPeak(InputAudio,  InOutWaveInfo.SampleRate, InOutWaveInfo.NumChannels);
+		break;
+	case ENormalizationMode::COUNT:
+	default:
+		return;
+	}
+
+	const float TargetOffset = Target - PeakDecibelValue;
+	const float MakeupGain = FMath::Clamp(TargetOffset, -MaxGain, MaxGain);
+
+	if (MakeupGain != 0.f)
+	{
+		Audio::ArrayMultiplyByConstantInPlace(InputAudio, Audio::ConvertToLinear(MakeupGain));
+	}
+}
+
+Audio::FTransformationPtr UWaveformTransformationNormalize::CreateTransformation() const
+{
+	return MakeUnique<FWaveTransformationNormalize>(Target, MaxGain, Mode);
+}
\ No newline at end of file
diff --git a/Engine/Plugins/Editor/WaveformEditor/Source/WaveformTransformations/Private/WaveformTransformationTrimFade.cpp b/Engine/Plugins/Editor/WaveformEditor/Source/WaveformTransformations/Private/WaveformTransformationTrimFade.cpp
new file mode 100644
--- /dev/null
+++ b/Engine/Plugins/Editor/WaveformEditor/Source/WaveformTransformations/Private/WaveformTransformationTrimFade.cpp
@@ -0,0 +1,129 @@
+// Copyright Epic Games, Inc. All Rights Reserved.
+
+#include "WaveformTransformationTrimFade.h"
+
+static void ApplyFadeIn(Audio::FAlignedFloatBuffer& InputAudio, const float FadeLength, const float FadeCurve, const int32 NumChannels, const float SampleRate)
+{
+	check(NumChannels > 0);
+
+	if(InputAudio.Num() < NumChannels || FadeLength < SMALL_NUMBER)
+	{
+		return;
+	}
+	
+	const int32 FadeNumFrames = FMath::Min(FadeLength * SampleRate / NumChannels, InputAudio.Num() / NumChannels);
+	float* InputPtr = InputAudio.GetData();
+	
+	for(int32 FrameIndex = 0;FrameIndex < FadeNumFrames; ++FrameIndex)
+	{
+		const float FadeFraction = (float)FrameIndex / FadeNumFrames;
+		const float EnvValue = FMath::Pow(FadeFraction, FadeCurve);
+		
+		for(int32 ChannelIt = 0; ChannelIt < NumChannels; ++ChannelIt)
+		{
+			*(InputPtr + ChannelIt) *= EnvValue;
+		}
+
+		InputPtr += NumChannels;
+	}
+}
+
+static void ApplyFadeOut(Audio::FAlignedFloatBuffer& InputAudio, const float FadeLength, const float FadeCurve, const int32 NumChannels, const float SampleRate)
+{
+	check(NumChannels > 0);
+
+	if(InputAudio.Num() < NumChannels || FadeLength < SMALL_NUMBER)
+	{
+		return;
+	}
+	
+	const int32 FadeNumFrames = FMath::Min(FadeLength * SampleRate / NumChannels, InputAudio.Num() / NumChannels);
+	const int32 StartSampleIndex = InputAudio.Num() - (FadeNumFrames * NumChannels);
+	float* InputPtr = &InputAudio[StartSampleIndex];
+	
+	for(int32 FrameIndex = 0;FrameIndex < FadeNumFrames; ++FrameIndex)
+	{
+		const float FadeFraction = (float)FrameIndex / FadeNumFrames;
+		const float EnvValue = 1.f - FMath::Pow(FadeFraction, FadeCurve);
+		
+		for(int32 ChannelIt = 0; ChannelIt < NumChannels; ++ChannelIt)
+		{
+			*(InputPtr + ChannelIt) *= EnvValue;
+		}
+
+		InputPtr += NumChannels;
+	}
+}
+
+FWaveTransformationTrimFade::FWaveTransformationTrimFade(double InStartTime, double InEndTime, float InStartFadeTime, float InStartFadeCurve, float InEndFadeTime, float InEndFadeCurve)
+	: StartTime(InStartTime)
+	, EndTime(InEndTime)
+	, StartFadeTime(InStartFadeTime)
+	, StartFadeCurve(FMath::Max(InStartFadeCurve, 0.f))
+	, EndFadeTime(InEndFadeTime)
+	, EndFadeCurve(FMath::Max(InEndFadeCurve, 0.f)) {}
+
+void FWaveTransformationTrimFade::ProcessAudio(Audio::FWaveformTransformationWaveInfo& InOutWaveInfo) const
+{
+	check(InOutWaveInfo.SampleRate > 0.f && InOutWaveInfo.Audio != nullptr);
+
+	Audio::FAlignedFloatBuffer& InputAudio = *InOutWaveInfo.Audio;
+
+	if(InputAudio.Num() == 0)
+	{
+		return;
+	}
+
+	const int32 StartSample = FMath::RoundToInt32(StartTime * InOutWaveInfo.SampleRate) * InOutWaveInfo.NumChannels;
+	int32 EndSample = InputAudio.Num();
+
+	if(EndTime > 0.f)
+	{
+		const int32 IndexOfLastSampleInFrame = InOutWaveInfo.NumChannels - 1;
+		EndSample = FMath::RoundToInt32(EndTime * InOutWaveInfo.SampleRate) * InOutWaveInfo.NumChannels + IndexOfLastSampleInFrame;
+		EndSample = FMath::Min(EndSample, InputAudio.Num());
+	}
+
+	const int32 FinalSize = EndSample - StartSample;
+
+	InOutWaveInfo.StartFrameOffset = StartSample - (StartSample % InOutWaveInfo.NumChannels);
+	InOutWaveInfo.NumEditedSamples = FinalSize;
+	
+	if (FinalSize > InputAudio.Num() || FinalSize <= 0)
+	{
+		return;
+	}
+
+	const bool bProcessFades = StartFadeTime > 0.f || EndFadeTime > 0.f;
+
+	if (!bProcessFades && FinalSize == InputAudio.Num())
+	{
+		return;
+	}
+
+	TArray<float> TempBuffer;
+	TempBuffer.AddUninitialized(FinalSize);
+
+	FMemory::Memcpy(TempBuffer.GetData(), &InputAudio[StartSample], FinalSize * sizeof(float));
+
+	InputAudio.Empty();
+	InputAudio.AddUninitialized(FinalSize);
+
+	FMemory::Memcpy(InputAudio.GetData(), TempBuffer.GetData(), FinalSize * sizeof(float));
+
+
+	if(StartFadeTime > 0.f)
+	{
+		ApplyFadeIn(InputAudio, StartFadeTime, StartFadeCurve, InOutWaveInfo.NumChannels, InOutWaveInfo.SampleRate);
+	}
+	
+	if(EndFadeTime > 0.f)
+	{
+		ApplyFadeOut(InputAudio, EndFadeTime, EndFadeCurve, InOutWaveInfo.NumChannels, InOutWaveInfo.SampleRate);
+	}
+}
+
+Audio::FTransformationPtr UWaveformTransformationTrimFade::CreateTransformation() const
+{
+	return MakeUnique<FWaveTransformationTrimFade>(StartTime, EndTime, StartFadeTime, StartFadeCurve, EndFadeTime, EndFadeCurve);
+}
diff --git a/Engine/Plugins/Editor/WaveformEditor/Source/WaveformTransformations/Private/WaveformTransformationsModule.cpp b/Engine/Plugins/Editor/WaveformEditor/Source/WaveformTransformations/Private/WaveformTransformationsModule.cpp
new file mode 100644
--- /dev/null
+++ b/Engine/Plugins/Editor/WaveformEditor/Source/WaveformTransformations/Private/WaveformTransformationsModule.cpp
@@ -0,0 +1,6 @@
+// Copyright Epic Games, Inc. All Rights Reserved.
+
+#include "WaveformTransformationsModule.h"
+#include "Modules/ModuleManager.h"
+
+IMPLEMENT_MODULE(FWaveformTransformationsModule, WaveformTransformations);
diff --git a/Engine/Plugins/Editor/WaveformEditor/Source/WaveformTransformations/Public/WaveformTransformationEffectChain.h b/Engine/Plugins/Editor/WaveformEditor/Source/WaveformTransformations/Public/WaveformTransformationEffectChain.h
new file mode 100644
--- /dev/null
+++ b/Engine/Plugins/Editor/WaveformEditor/Source/WaveformTransformations/Public/WaveformTransformationEffectChain.h
@@ -0,0 +1,39 @@
+// Copyright Epic Games, Inc. All Rights Reserved.
+
+#pragma once
+
+#include "AudioExtensions/Public/IWaveformTransformation.h"
+#include "Factories/SoundSourceEffectFactory.h"
+#include "UObject/Object.h"
+#include "UObject/StrongObjectPtr.h"
+#include "WaveformTransformationEffectChain.generated.h"
+
+class WAVEFORMTRANSFORMATIONS_API FWaveTransformationEffectChain : public Audio::IWaveTransformation
+{
+public:
+	FWaveTransformationEffectChain(TArray<TObjectPtr<class USoundEffectSourcePreset>>& InEffectPresets);
+
+	virtual void ProcessAudio(Audio::FWaveformTransformationWaveInfo& InOutWaveInfo) const override;
+
+	virtual bool SupportsRealtimePreview() const override { return true; }
+	virtual bool CanChangeFileLength() const override { return true; }
+
+private:
+	TArray<TStrongObjectPtr<class USoundEffectSourcePreset>> Presets;
+};
+
+UCLASS()
+class WAVEFORMTRANSFORMATIONS_API UWaveformTransformationEffectChain : public UWaveformTransformationBase
+{
+	GENERATED_BODY()
+
+	UPROPERTY(EditAnywhere, Category = "Trim")
+	TObjectPtr<class USoundEffectSourcePresetChain> EffectChain;
+
+	UPROPERTY(EditAnywhere, Instanced, Category = "Trim")
+	TArray<TObjectPtr<class USoundEffectSourcePreset>> InlineEffects;
+	
+public:
+	
+	virtual Audio::FTransformationPtr CreateTransformation() const override;
+};
\ No newline at end of file
diff --git a/Engine/Plugins/Editor/WaveformEditor/Source/WaveformTransformations/Public/WaveformTransformationNormalize.h b/Engine/Plugins/Editor/WaveformEditor/Source/WaveformTransformations/Public/WaveformTransformationNormalize.h
new file mode 100644
--- /dev/null
+++ b/Engine/Plugins/Editor/WaveformEditor/Source/WaveformTransformations/Public/WaveformTransformationNormalize.h
@@ -0,0 +1,50 @@
+// Copyright Epic Games, Inc. All Rights Reserved.
+
+#pragma once
+
+#include "AudioExtensions/Public/IWaveformTransformation.h"
+#include "UObject/Object.h"
+#include "WaveformTransformationNormalize.generated.h"
+
+UENUM()
+enum class ENormalizationMode : uint8
+{
+	Peak,
+	RMS,
+	DWeightedLoudness,
+	COUNT UMETA(Hidden)
+};
+
+class WAVEFORMTRANSFORMATIONS_API FWaveTransformationNormalize : public Audio::IWaveTransformation
+{
+public:
+	FWaveTransformationNormalize(float InTarget, float InMaxGain, ENormalizationMode InMode);
+
+	virtual void ProcessAudio(Audio::FWaveformTransformationWaveInfo& InOutWaveInfo) const override;
+
+private:
+	float Target = 0.f;
+	float MaxGain = 0.f;
+	ENormalizationMode Mode = ENormalizationMode::Peak;
+};
+
+UCLASS()
+class WAVEFORMTRANSFORMATIONS_API UWaveformTransformationNormalize : public UWaveformTransformationBase
+{
+	GENERATED_BODY()
+
+	// Target maximum volume for this soundwave, in decibels
+	UPROPERTY(EditAnywhere, Category="Normalization", meta=(ClampMax = 0.0))
+	float Target = 0.f;
+
+	// Will not apply more gain than this, even if the sound is very quiet
+	UPROPERTY(EditAnywhere, Category="Normalization", meta=(ClampMin = 0.0))
+	float MaxGain = 0.f;
+
+	// what type of analysis to run to find the peak value
+	UPROPERTY(EditAnywhere, Category="Normalization")
+	ENormalizationMode Mode;
+	
+public:
+	virtual Audio::FTransformationPtr CreateTransformation() const override;
+};
\ No newline at end of file
diff --git a/Engine/Plugins/Editor/WaveformEditor/Source/WaveformTransformations/Public/WaveformTransformationTrimFade.h b/Engine/Plugins/Editor/WaveformEditor/Source/WaveformTransformations/Public/WaveformTransformationTrimFade.h
new file mode 100644
--- /dev/null
+++ b/Engine/Plugins/Editor/WaveformEditor/Source/WaveformTransformations/Public/WaveformTransformationTrimFade.h
@@ -0,0 +1,103 @@
+// Copyright Epic Games, Inc. All Rights Reserved.
+
+#pragma once
+
+#include "UObject/Object.h"
+#include "IWaveformTransformation.h"
+
+#include "WaveformTransformationTrimFade.generated.h"
+
+class WAVEFORMTRANSFORMATIONS_API FWaveTransformationTrimFade : public Audio::IWaveTransformation
+{
+public:
+	explicit FWaveTransformationTrimFade(double InStartTime, double InEndTime, float InStartFadeTime, float InStartFadeCurve, float InEndFadeTime, float InEndFadeCurve);
+	virtual void ProcessAudio(Audio::FWaveformTransformationWaveInfo& InOutWaveInfo) const override;
+
+	virtual bool CanChangeFileLength() const override { return true; }
+
+private:
+	double StartTime = 0.0;
+	double EndTime = 0.0;
+
+	float StartFadeTime = 0.f;
+	float StartFadeCurve = 0.f;
+
+	float EndFadeTime = 0.f;
+	float EndFadeCurve = 0.f;
+};
+
+UCLASS()
+class WAVEFORMTRANSFORMATIONS_API UWaveformTransformationTrimFade : public UWaveformTransformationBase
+{
+	GENERATED_BODY()
+
+	UPROPERTY(EditAnywhere, Category = "Trim", meta=(ClampMin = 0.0))
+	double StartTime = 0.0;
+
+	UPROPERTY(EditAnywhere, Category = "Trim")
+	double EndTime = -1.0;
+
+	UPROPERTY(EditAnywhere, Category = "Fade", meta=(ClampMin = 0.0, DisplayName = "Fade-In Duration"))
+	float StartFadeTime = 0.f;
+
+	UPROPERTY(EditAnywhere, Category = "Fade", meta=(ClampMin = 0.0, ClampMax = 10.0, DisplayName = "Fade-In Curve"))
+	float StartFadeCurve = 1.f;
+
+	UPROPERTY(EditAnywhere, Category = "Fade", meta=(ClampMin = 0.0, DisplayName = "Fade-Out Duration"))
+	float EndFadeTime = 0.f;
+
+	UPROPERTY(EditAnywhere, Category = "Fade", meta=(ClampMin = 0.0, ClampMax = 10.0, DisplayName = "Fade-Out Curve"))
+	float EndFadeCurve = 1.f;
+
+	
+public:
+	virtual Audio::FTransformationPtr CreateTransformation() const override;
+};
+
+// class UBaseClass;
+//
+// namespace  MyNamespace
+// {
+// 	class IBaseInterface;
+// }
+
+
+
+// // no good
+// #include "IBaseInterface.h" // in other module but that module is listed in the build.cs so we good there
+// #include "DerivedClass.generated.h"
+//
+// class PLUGIN_API FDerivedClass : public MyNamespace::IBaseInterface // error: Class 'MyNamespace' Not Found
+// {
+// public:
+// 	virtual void DoThing() const;
+// };
+//
+// UCLASS()
+// class PLUGIN_API UDerivedClass : public UBaseClass
+// {
+// 	GENERATED_BODY()
+// public:
+// 	TSharedPtr<MyNamespace::IBaseInterface> CreateFObject();
+// };
+//
+// // fine
+// #include "IBaseInterface.h"
+// #include "DerivedClass.generated.h"
+//
+// namespace MyNamespace
+// {
+// 	class PLUGIN_API FDerivedClass : public MyNamespace::IBaseInterface
+// 	{
+// 	public:
+// 		virtual void DoThing() const;
+// 	};
+// }
+//
+// UCLASS()
+// class PLUGIN_API UDerivedClass : public UBaseClass
+// {
+// 	GENERATED_BODY()
+// public:
+// 	TSharedPtr<MyNamespace::IBaseInterface> CreateFObject();
+// };
\ No newline at end of file
diff --git a/Engine/Plugins/Editor/WaveformEditor/Source/WaveformTransformations/Public/WaveformTransformationsModule.h b/Engine/Plugins/Editor/WaveformEditor/Source/WaveformTransformations/Public/WaveformTransformationsModule.h
new file mode 100644
--- /dev/null
+++ b/Engine/Plugins/Editor/WaveformEditor/Source/WaveformTransformations/Public/WaveformTransformationsModule.h
@@ -0,0 +1,9 @@
+// Copyright Epic Games, Inc. All Rights Reserved.
+
+#pragma once
+
+#include "Modules/ModuleInterface.h"
+
+class FWaveformTransformationsModule : public IModuleInterface
+{
+};
diff --git a/Engine/Plugins/Editor/WaveformEditor/Source/WaveformTransformations/WaveformTransformations.build.cs b/Engine/Plugins/Editor/WaveformEditor/Source/WaveformTransformations/WaveformTransformations.build.cs
new file mode 100644
--- /dev/null
+++ b/Engine/Plugins/Editor/WaveformEditor/Source/WaveformTransformations/WaveformTransformations.build.cs
@@ -0,0 +1,25 @@
+// Copyright Epic Games, Inc. All Rights Reserved.
+
+using UnrealBuildTool;
+
+namespace UnrealBuildTool.Rules
+{
+	public class WaveformTransformations : ModuleRules
+	{
+		public WaveformTransformations(ReadOnlyTargetRules Target) : base(Target)
+		{
+			PCHUsage = PCHUsageMode.UseExplicitOrSharedPCHs;
+
+			PublicDependencyModuleNames.AddRange(
+				new string[] {
+					"AudioExtensions",
+					"AudioSynesthesiaCore",
+					"Core", 
+					"CoreUObject",
+					"Engine",
+					"SignalProcessing"
+				}
+			);
+		}
+	}
+}
\ No newline at end of file
diff --git a/Engine/Plugins/Editor/WaveformEditor/Source/WaveformTransformationsWidgets/Private/SWaveformTransformationDurationHighlight.h b/Engine/Plugins/Editor/WaveformEditor/Source/WaveformTransformationsWidgets/Private/SWaveformTransformationDurationHighlight.h
new file mode 100644
--- /dev/null
+++ b/Engine/Plugins/Editor/WaveformEditor/Source/WaveformTransformationsWidgets/Private/SWaveformTransformationDurationHighlight.h
@@ -0,0 +1,85 @@
+// Copyright Epic Games, Inc. All Rights Reserved.
+
+#pragma once
+
+#include "SWaveformTransformationRenderLayer.h"
+#include "WaveformEditorRenderData.h"
+#include "WaveformEditorTransportCoordinator.h"
+
+namespace WaveformTransformationDurationHiglightParams
+{
+	const FLazyName BackgroundBrushName = TEXT("WhiteBrush");
+	const FLinearColor BoxColor = FLinearColor(0.f, 0.f, 0.f, 0.7f);
+}
+
+class SWaveformTransformationDurationHiglight : public SWaveformTransformationRenderLayer
+{
+public:
+	SLATE_BEGIN_ARGS(SWaveformTransformationDurationHiglight) {}
+		SLATE_DEFAULT_SLOT(FArguments, InArgs)
+	SLATE_END_ARGS()
+
+	void Construct(const FArguments& InArgs, TSharedRef<FWaveformEditorRenderData> InWaveformRenderData)
+	{
+		WaveformRenderData = InWaveformRenderData;
+	}
+
+	int32 OnPaint(const FPaintArgs& Args, const FGeometry& AllottedGeometry, const FSlateRect& MyCullingRect, FSlateWindowElementList& OutDrawElements, int32 LayerId, const FWidgetStyle& InWidgetStyle, bool bParentEnabled) const override
+	{
+		const float StartTimeRatio = (float)TransformationWaveInfo.StartFrameOffset / WaveformRenderData->GetOriginalWaveformFrames() / TransformationWaveInfo.NumChannels;
+		const uint32 EndSample = TransformationWaveInfo.StartFrameOffset + TransformationWaveInfo.NumEditedSamples;
+		const float EndTimeRatio = (float)EndSample / WaveformRenderData->GetOriginalWaveformFrames() / TransformationWaveInfo.NumChannels;
+
+		const float LeftBoundXRatio = (StartTimeRatio - WaveformDisplayRange.GetLowerBoundValue()) / ZoomRatio;
+		const float RightBoundXRatio = (EndTimeRatio - WaveformDisplayRange.GetLowerBoundValue()) / ZoomRatio;
+
+		const bool bRenderLeftBox = LeftBoundXRatio >= 0.f;
+		const bool bRenderRightBox = RightBoundXRatio <= 1.f;
+
+		if (bRenderLeftBox)
+		{
+			const float RightMarginX = LeftBoundXRatio * AllottedGeometry.Size.X;
+
+			FSlateDrawElement::MakeBox(
+				OutDrawElements,
+				++LayerId,
+				AllottedGeometry.ToPaintGeometry(FVector2D(0.f, 0.f), FVector2D(RightMarginX, AllottedGeometry.Size.Y)),
+				FAppStyle::GetBrush(WaveformTransformationDurationHiglightParams::BackgroundBrushName),
+				ESlateDrawEffect::None,
+				WaveformTransformationDurationHiglightParams::BoxColor);
+
+		}
+
+		if (bRenderRightBox)
+		{
+			const float LeftMarginX = RightBoundXRatio * AllottedGeometry.Size.X;
+
+			FSlateDrawElement::MakeBox(
+				OutDrawElements,
+				++LayerId,
+				AllottedGeometry.ToPaintGeometry(FVector2D(LeftMarginX, 0.f), FVector2D(AllottedGeometry.Size.X, AllottedGeometry.Size.Y)),
+				FAppStyle::GetBrush(WaveformTransformationDurationHiglightParams::BackgroundBrushName),
+				ESlateDrawEffect::None,
+				WaveformTransformationDurationHiglightParams::BoxColor);
+
+		}
+
+		return LayerId;
+	}
+
+	void UpdateDisplayRange(const TRange<float> NewDisplayRange)
+	{
+		WaveformDisplayRange = NewDisplayRange;
+	}
+
+	void OnZoomLevelChanged(const uint8 NewLevel)
+	{
+		ZoomRatio = FMath::Clamp(NewLevel / 100.f, UE_KINDA_SMALL_NUMBER, 1.f);
+	}
+
+private:
+	TSharedPtr<FWaveformEditorRenderData> WaveformRenderData = nullptr;
+
+	TRange<float> WaveformDisplayRange = TRange<float>::Inclusive(0.f, 1.f);
+	float ZoomRatio = 1.f;
+};
\ No newline at end of file
diff --git a/Engine/Plugins/Editor/WaveformEditor/Source/WaveformTransformationsWidgets/Private/SWaveformTransformationRenderLayer.cpp b/Engine/Plugins/Editor/WaveformEditor/Source/WaveformTransformationsWidgets/Private/SWaveformTransformationRenderLayer.cpp
new file mode 100644
--- /dev/null
+++ b/Engine/Plugins/Editor/WaveformEditor/Source/WaveformTransformationsWidgets/Private/SWaveformTransformationRenderLayer.cpp
@@ -0,0 +1,41 @@
+// Copyright Epic Games, Inc. All Rights Reserved.
+
+#include "SWaveformTransformationRenderLayer.h"
+
+#include "Styling/AppStyle.h"
+
+void SWaveformTransformationRenderLayer::PrivateRegisterAttributes(FSlateAttributeInitializer& AttributeInitializer)
+{
+}
+
+void SWaveformTransformationRenderLayer::Construct(const FArguments& InArgs)
+{
+}
+
+int32 SWaveformTransformationRenderLayer::OnPaint(const FPaintArgs& Args, const FGeometry& AllottedGeometry, const FSlateRect& MyCullingRect, FSlateWindowElementList& OutDrawElements, int32 LayerId, const FWidgetStyle& InWidgetStyle, bool bParentEnabled) const
+{
+
+	auto DisplayString = FString::Printf(TEXT("basic transform layer %d"), LayerId);
+
+	FSlateDrawElement::MakeText(
+		OutDrawElements,
+		++LayerId,
+		AllottedGeometry.ToPaintGeometry(),
+		DisplayString,
+		FAppStyle::GetFontStyle("Regular"),
+		ESlateDrawEffect::None,
+		FLinearColor::Green
+	);
+
+	return LayerId;
+}
+
+FVector2D SWaveformTransformationRenderLayer::ComputeDesiredSize(float) const
+{
+	return FVector2D(1280, 720);
+}
+
+void SWaveformTransformationRenderLayer::SetTransformationWaveInfo(FWaveformTransformationRenderLayerInfo InWaveInfo)
+{
+	TransformationWaveInfo = InWaveInfo;
+}
diff --git a/Engine/Plugins/Editor/WaveformEditor/Source/WaveformTransformationsWidgets/Private/SWaveformTransformationTrimFadeLayer.h b/Engine/Plugins/Editor/WaveformEditor/Source/WaveformTransformationsWidgets/Private/SWaveformTransformationTrimFadeLayer.h
new file mode 100644
--- /dev/null
+++ b/Engine/Plugins/Editor/WaveformEditor/Source/WaveformTransformationsWidgets/Private/SWaveformTransformationTrimFadeLayer.h
@@ -0,0 +1,93 @@
+// Copyright Epic Games, Inc. All Rights Reserved.
+
+#pragma once
+
+#include "SWaveformTransformationRenderLayer.h"
+#include "WaveformEditorRenderData.h"
+#include "WaveformTransformationTrimFade.h"
+
+class SWaveformTransformationTrimFadeLayer : public SWaveformTransformationRenderLayer
+{
+public:
+	SLATE_BEGIN_ARGS(SWaveformTransformationTrimFadeLayer) {}
+		SLATE_DEFAULT_SLOT(FArguments, InArgs)
+	SLATE_END_ARGS()
+
+	void Construct(const FArguments& InArgs, TObjectPtr<UWaveformTransformationTrimFade> InTransformToRender, TSharedRef<FWaveformEditorRenderData> InWaveformRenderData)
+	{
+		TrimFadeTransform = InTransformToRender;
+		WaveformRenderData = InWaveformRenderData;
+	}
+
+	int32 OnPaint(const FPaintArgs& Args, const FGeometry& AllottedGeometry, const FSlateRect& MyCullingRect, FSlateWindowElementList& OutDrawElements, int32 LayerId, const FWidgetStyle& InWidgetStyle, bool bParentEnabled) const override
+	{
+		const float StartTimeRatio = (float)TransformationWaveInfo.StartFrameOffset / WaveformRenderData->GetOriginalWaveformFrames() / TransformationWaveInfo.NumChannels;
+		const uint32 EndSample = TransformationWaveInfo.StartFrameOffset + TransformationWaveInfo.NumEditedSamples;
+		const float EndTimeRatio = (float)EndSample / WaveformRenderData->GetOriginalWaveformFrames() / TransformationWaveInfo.NumChannels;
+
+		const float TrimStartBarXRatio = (StartTimeRatio - WaveformDisplayRange.GetLowerBoundValue()) / ZoomRatio;
+		const float TrimEndBarXRatio = (EndTimeRatio - WaveformDisplayRange.GetLowerBoundValue()) / ZoomRatio;
+
+		const bool bRenderLowerBound = TrimStartBarXRatio >= 0.f;
+		const bool bRenderUpperBound = TrimEndBarXRatio <= 1.f;
+		
+		TArray<FVector2D> LinePoints;
+		LinePoints.SetNumUninitialized(2);
+
+		if (bRenderLowerBound)
+		{
+			const float LowerBoundX = TrimStartBarXRatio * AllottedGeometry.Size.X;
+
+			LinePoints[0] = FVector2D(LowerBoundX, 0.f);
+			LinePoints[1] = FVector2D(LowerBoundX, AllottedGeometry.Size.Y);
+
+			FSlateDrawElement::MakeLines(
+				OutDrawElements,
+				++LayerId,
+				AllottedGeometry.ToPaintGeometry(),
+				LinePoints,
+				ESlateDrawEffect::None,
+				FLinearColor::Green,
+				false
+			);
+		}
+
+		if (bRenderUpperBound)
+		{
+			const float UpperBoundX = TrimEndBarXRatio * AllottedGeometry.Size.X;
+
+			LinePoints[0] = FVector2D(UpperBoundX, 0.f);
+			LinePoints[1] = FVector2D(UpperBoundX, AllottedGeometry.Size.Y);
+
+			FSlateDrawElement::MakeLines(
+				OutDrawElements,
+				++LayerId,
+				AllottedGeometry.ToPaintGeometry(),
+				LinePoints,
+				ESlateDrawEffect::None,
+				FLinearColor::Red,
+				false
+			);
+		}
+
+
+		return LayerId;
+	}
+
+	void UpdateDisplayRange(const TRange<float> NewDisplayRange)
+	{
+		WaveformDisplayRange = NewDisplayRange;
+	}
+
+	void OnZoomLevelChanged(const uint8 NewLevel)
+	{
+		ZoomRatio = FMath::Clamp(NewLevel / 100.f, UE_KINDA_SMALL_NUMBER, 1.f);
+	}
+
+private:
+	TObjectPtr<UWaveformTransformationTrimFade> TrimFadeTransform = nullptr;
+	TSharedPtr<FWaveformEditorRenderData> WaveformRenderData = nullptr;
+
+	TRange<float> WaveformDisplayRange = TRange<float>::Inclusive(0.f, 1.f);
+	float ZoomRatio = 1.f;
+};
\ No newline at end of file
diff --git a/Engine/Plugins/Editor/WaveformEditor/Source/WaveformTransformationsWidgets/Private/WaveformTransformationRenderLayerFactory.cpp b/Engine/Plugins/Editor/WaveformEditor/Source/WaveformTransformationsWidgets/Private/WaveformTransformationRenderLayerFactory.cpp
new file mode 100644
--- /dev/null
+++ b/Engine/Plugins/Editor/WaveformEditor/Source/WaveformTransformationsWidgets/Private/WaveformTransformationRenderLayerFactory.cpp
@@ -0,0 +1,49 @@
+// Copyright Epic Games, Inc. All Rights Reserved.
+
+#include "WaveformTransformationRenderLayerFactory.h"
+
+#include "IWaveformTransformation.h"
+#include "Styling/AppStyle.h"
+#include "SWaveformTransformationDurationHighlight.h"
+#include "SWaveformTransformationRenderLayer.h"
+#include "SWaveformTransformationTrimFadeLayer.h"
+#include "WaveformEditorTransportCoordinator.h"
+#include "WaveformEditorZoomController.h"
+#include "WaveformTransformationTrimFade.h"
+
+FWaveformTransformationRenderLayerFactory::FWaveformTransformationRenderLayerFactory(TSharedRef<FWaveformEditorRenderData> InWaveformRenderData, TSharedRef<FWaveformEditorTransportCoordinator> InTransportCoordinator, TSharedRef<FWaveformEditorZoomController> InZoomController)
+	: WaveformRenderData(InWaveformRenderData)
+	, TransportCoordinator(InTransportCoordinator)
+	, ZoomController(InZoomController)
+{
+}
+
+
+TSharedPtr<SWaveformTransformationRenderLayer> FWaveformTransformationRenderLayerFactory::Create(TObjectPtr<UWaveformTransformationBase> InTransformationToRender) const
+{
+	TSharedPtr<SWaveformTransformationRenderLayer> OutWidget = nullptr;
+	UClass* TransformationClass = InTransformationToRender->GetClass();
+
+	if (TransformationClass == UWaveformTransformationTrimFade::StaticClass())
+	{
+		TSharedPtr<SWaveformTransformationTrimFadeLayer> TrimFadeLayer = SNew(SWaveformTransformationTrimFadeLayer, Cast<UWaveformTransformationTrimFade>(InTransformationToRender), WaveformRenderData.ToSharedRef());
+		TrimFadeLayer->OnZoomLevelChanged(ZoomController->GetZoomRatio());
+		TrimFadeLayer->UpdateDisplayRange(TransportCoordinator->GetDisplayRange());
+		TransportCoordinator->OnDisplayRangeUpdated.AddSP(TrimFadeLayer.Get(), &SWaveformTransformationTrimFadeLayer::UpdateDisplayRange);
+		ZoomController->OnZoomRatioChanged.AddSP(TrimFadeLayer.Get(), &SWaveformTransformationTrimFadeLayer::OnZoomLevelChanged);
+		OutWidget = TrimFadeLayer;
+	}
+
+	return OutWidget;
+}
+
+TSharedPtr<SWaveformTransformationRenderLayer> FWaveformTransformationRenderLayerFactory::CreateDurationHiglightLayer() const
+{
+	TSharedPtr<SWaveformTransformationDurationHiglight> DurationHiglightLayer = SNew(SWaveformTransformationDurationHiglight, WaveformRenderData.ToSharedRef());
+	DurationHiglightLayer->OnZoomLevelChanged(ZoomController->GetZoomRatio());
+	DurationHiglightLayer->UpdateDisplayRange(TransportCoordinator->GetDisplayRange());
+	TransportCoordinator->OnDisplayRangeUpdated.AddSP(DurationHiglightLayer.Get(), &SWaveformTransformationDurationHiglight::UpdateDisplayRange);
+	ZoomController->OnZoomRatioChanged.AddSP(DurationHiglightLayer.Get(), &SWaveformTransformationDurationHiglight::OnZoomLevelChanged);
+
+	return DurationHiglightLayer;
+}
diff --git a/Engine/Plugins/Editor/WaveformEditor/Source/WaveformTransformationsWidgets/Private/WaveformTransformationsRenderManager.cpp b/Engine/Plugins/Editor/WaveformEditor/Source/WaveformTransformationsWidgets/Private/WaveformTransformationsRenderManager.cpp
new file mode 100644
--- /dev/null
+++ b/Engine/Plugins/Editor/WaveformEditor/Source/WaveformTransformationsWidgets/Private/WaveformTransformationsRenderManager.cpp
@@ -0,0 +1,184 @@
+// Copyright Epic Games, Inc. All Rights Reserved.
+
+#include "WaveformTransformationsRenderManager.h"
+
+#include "DSP/BufferVectorOperations.h"
+#include "DSP/FloatArrayMath.h"
+#include "Sound/SoundWave.h"
+#include "SWaveformTransformationRenderLayer.h"
+#include "Widgets/SWidget.h"
+
+FWaveformTransformationsRenderManager::FWaveformTransformationsRenderManager(TObjectPtr<USoundWave> InSoundWave, TSharedRef<FWaveformEditorRenderData> InWaveformRenderData, TSharedRef<FWaveformEditorTransportCoordinator> InTransportCoordinator, TSharedRef<FWaveformEditorZoomController> InZoomController)
+	: SoundWaveToRender(InSoundWave)
+	, LayersFactory(MakeUnique<FWaveformTransformationRenderLayerFactory>(InWaveformRenderData, InTransportCoordinator, InZoomController))
+{
+	GenerateLayersChain();	
+}
+
+void FWaveformTransformationsRenderManager::GenerateLayersChain()
+{
+	check(SoundWaveToRender);
+
+	TransformationsToRender.Empty();
+	RenderLayers.Empty();
+
+	if (SoundWaveToRender->Transformations.Num() > 0)
+	{
+		for (const TObjectPtr<UWaveformTransformationBase> Transformation : SoundWaveToRender->Transformations)
+		{
+			if (Transformation != nullptr)
+			{
+				TransformationsToRender.Add(Transformation);
+				
+				TSharedPtr<SWaveformTransformationRenderLayer> TransformationUI = LayersFactory->Create(Transformation);
+				RenderLayers.Add(TransformationUI);
+				
+			}
+		}
+
+		CreateDurationHighlightLayer();
+
+	}
+
+	TSharedPtr<SWidget>* WidgetPtr = reinterpret_cast<TSharedPtr<SWidget>*>(RenderLayers.GetData());
+	RenderLayersWidgetView = MakeArrayView(WidgetPtr, RenderLayers.Num());
+
+	OnLayersChainGenerated.Broadcast(RenderLayersWidgetView.GetData(), RenderLayers.Num());
+}
+
+void FWaveformTransformationsRenderManager::UpdateRenderElements()
+{
+	GenerateRenderDataInternal();
+}
+
+
+
+void FWaveformTransformationsRenderManager::GenerateRenderDataInternal()
+{
+	check(SoundWaveToRender);
+
+	RawPCMData.Empty();
+	uint16 NumChannels = 0;
+	uint32 SampleRate = 0;
+
+	if (!SoundWaveToRender->GetImportedSoundWaveData(RawPCMData, SampleRate, NumChannels))
+	{
+		UE_LOG(LogAudio, Warning, TEXT("Failed to get transformations render data for: %s"), *SoundWaveToRender->GetPathName());
+		return;
+	}
+
+	uint32 NumOriginalSamples = RawPCMData.Num() * sizeof(uint8) / sizeof(int16);
+	uint32 FirstEditedSample = 0;
+	uint32 LastEditedSample = NumOriginalSamples;
+
+	if (TransformationsToRender.Num() > 0)
+	{
+		Audio::FWaveformTransformationWaveInfo TransformationInfo;
+
+
+		Audio::FAlignedFloatBuffer TransformationsBuffer;
+		Audio::FAlignedFloatBuffer OutputBuffer;
+		TransformationsBuffer.SetNumUninitialized(NumOriginalSamples);
+
+		Audio::ArrayPcm16ToFloat(MakeArrayView((int16*)RawPCMData.GetData(), NumOriginalSamples), TransformationsBuffer);
+		OutputBuffer = TransformationsBuffer;
+
+		TransformationInfo.Audio = &TransformationsBuffer;
+		TransformationInfo.NumChannels = NumChannels;
+		TransformationInfo.SampleRate = SampleRate;
+
+		TArray<Audio::FTransformationPtr> Transformations = CreateTransformations();
+
+		const bool bChainChangesFileLength = CanChainChangeFileLength(Transformations);
+
+		for (int32 i = 0; i < Transformations.Num(); ++i)
+		{
+			Transformations[i]->ProcessAudio(TransformationInfo);
+
+			FirstEditedSample += TransformationInfo.StartFrameOffset;
+			LastEditedSample = TransformationInfo.NumEditedSamples == 0 ? LastEditedSample : FirstEditedSample + TransformationInfo.NumEditedSamples;
+
+			if (RenderLayers[i])
+			{
+				FWaveformTransformationRenderLayerInfo RenderLayerInfo{ TransformationInfo.SampleRate, TransformationInfo.NumChannels, FirstEditedSample, TransformationInfo.NumEditedSamples };
+				RenderLayers[i]->SetTransformationWaveInfo(MoveTemp(RenderLayerInfo));
+			}
+
+			if (bChainChangesFileLength)
+			{
+				check(LastEditedSample > FirstEditedSample);
+				const int32 EditedBufferSize = LastEditedSample - FirstEditedSample;
+				FMemory::Memcpy(&OutputBuffer.GetData()[FirstEditedSample], TransformationsBuffer.GetData(), EditedBufferSize * sizeof(float));
+			}
+
+			TransformationInfo.StartFrameOffset = 0;
+			TransformationInfo.NumEditedSamples = 0;
+		}
+
+		check(DurationHiglightLayer)
+		FWaveformTransformationRenderLayerInfo DurationLayerInfo{ TransformationInfo.SampleRate, TransformationInfo.NumChannels, FirstEditedSample, LastEditedSample - FirstEditedSample };
+		DurationHiglightLayer->SetTransformationWaveInfo(MoveTemp(DurationLayerInfo));
+			
+		if (!bChainChangesFileLength)
+		{
+			OutputBuffer = TransformationsBuffer;
+		}
+		
+
+		if (const float MaxValue = Audio::ArrayMaxAbsValue(OutputBuffer) > 1.f)
+		{
+			Audio::ArrayMultiplyByConstantInPlace(OutputBuffer, 1.f / MaxValue);
+		}
+
+		SampleRate = TransformationInfo.SampleRate;
+		NumChannels = TransformationInfo.NumChannels;
+		NumOriginalSamples = OutputBuffer.Num();
+		
+		check(NumChannels > 0);
+		check(SampleRate > 0);
+
+		Audio::ArrayFloatToPcm16(OutputBuffer, MakeArrayView((int16*)RawPCMData.GetData(), NumOriginalSamples));
+	}
+
+	OnRenderDataGenerated.Broadcast(RawPCMData.GetData(), NumOriginalSamples, FirstEditedSample, LastEditedSample, SampleRate, NumChannels);
+}
+
+TArray<Audio::FTransformationPtr> FWaveformTransformationsRenderManager::CreateTransformations() const
+{
+	TArray<Audio::FTransformationPtr> TransformationPtrs;
+
+	for (const TObjectPtr<UWaveformTransformationBase> TransformationBase : TransformationsToRender)
+	{
+		TransformationPtrs.Add(TransformationBase->CreateTransformation());
+	}
+
+	return TransformationPtrs;
+}
+
+const bool FWaveformTransformationsRenderManager::CanChainChangeFileLength(const TArray<Audio::FTransformationPtr>& TransformationChain) const
+{
+	bool bCanChainChangeFileLength = false;
+
+	for (const Audio::FTransformationPtr& Transformation : TransformationChain)
+	{
+		bCanChainChangeFileLength |= Transformation->CanChangeFileLength();
+	}
+
+	return bCanChainChangeFileLength;
+}
+
+TArrayView<TSharedPtr<SWidget>> FWaveformTransformationsRenderManager::GetTransformLayers() const
+{
+	return RenderLayersWidgetView;
+}
+
+void FWaveformTransformationsRenderManager::CreateDurationHighlightLayer()
+{
+	if (!DurationHiglightLayer)
+	{
+		DurationHiglightLayer = LayersFactory->CreateDurationHiglightLayer();
+	}
+
+	RenderLayers.Add(DurationHiglightLayer);
+	DurationHiglightLayer->SetTransformationWaveInfo(FWaveformTransformationRenderLayerInfo());
+}
\ No newline at end of file
diff --git a/Engine/Plugins/Editor/WaveformEditor/Source/WaveformTransformationsWidgets/Private/WaveformTransformationsWidgetsModule.cpp b/Engine/Plugins/Editor/WaveformEditor/Source/WaveformTransformationsWidgets/Private/WaveformTransformationsWidgetsModule.cpp
new file mode 100644
--- /dev/null
+++ b/Engine/Plugins/Editor/WaveformEditor/Source/WaveformTransformationsWidgets/Private/WaveformTransformationsWidgetsModule.cpp
@@ -0,0 +1,6 @@
+// Copyright Epic Games, Inc. All Rights Reserved.
+
+#include "WaveformTransformationsWidgetsModule.h"
+#include "Modules/ModuleManager.h"
+
+IMPLEMENT_MODULE(FWaveformTransformationsWidgetsModule, WaveformTransformations);
\ No newline at end of file
diff --git a/Engine/Plugins/Editor/WaveformEditor/Source/WaveformTransformationsWidgets/Public/SWaveformTransformationRenderLayer.h b/Engine/Plugins/Editor/WaveformEditor/Source/WaveformTransformationsWidgets/Public/SWaveformTransformationRenderLayer.h
new file mode 100644
--- /dev/null
+++ b/Engine/Plugins/Editor/WaveformEditor/Source/WaveformTransformationsWidgets/Public/SWaveformTransformationRenderLayer.h
@@ -0,0 +1,36 @@
+// Copyright Epic Games, Inc. All Rights Reserved.
+
+#pragma once
+
+#include "Widgets/SCompoundWidget.h"
+
+struct WAVEFORMTRANSFORMATIONSWIDGETS_API FWaveformTransformationRenderLayerInfo
+{
+	float SampleRate = 0.f;
+	int32 NumChannels = 0;
+	uint32 StartFrameOffset = 0;
+	uint32 NumEditedSamples = 0;
+};
+
+class WAVEFORMTRANSFORMATIONSWIDGETS_API SWaveformTransformationRenderLayer : public SCompoundWidget
+{
+	SLATE_DECLARE_WIDGET(SWaveformTransformationRenderLayer, SCompoundWidget)
+
+public:
+	SLATE_BEGIN_ARGS(SWaveformTransformationRenderLayer) {}
+		SLATE_DEFAULT_SLOT(FArguments, InArgs)
+	SLATE_END_ARGS()
+
+	void Construct(const FArguments& InArgs);
+	
+	virtual int32 OnPaint(const FPaintArgs& Args, const FGeometry& AllottedGeometry, const FSlateRect& MyCullingRect, FSlateWindowElementList& OutDrawElements, int32 LayerId, const FWidgetStyle& InWidgetStyle, bool bParentEnabled) const override;
+	FVector2D ComputeDesiredSize(float) const override;
+
+	void SetTransformationWaveInfo(FWaveformTransformationRenderLayerInfo InWaveInfo);
+
+
+protected:
+	FWaveformTransformationRenderLayerInfo TransformationWaveInfo;
+
+
+};
\ No newline at end of file
diff --git a/Engine/Plugins/Editor/WaveformEditor/Source/WaveformTransformationsWidgets/Public/WaveformTransformationRenderLayerFactory.h b/Engine/Plugins/Editor/WaveformEditor/Source/WaveformTransformationsWidgets/Public/WaveformTransformationRenderLayerFactory.h
new file mode 100644
--- /dev/null
+++ b/Engine/Plugins/Editor/WaveformEditor/Source/WaveformTransformationsWidgets/Public/WaveformTransformationRenderLayerFactory.h
@@ -0,0 +1,31 @@
+// Copyright Epic Games, Inc. All Rights Reserved.
+
+#pragma once
+
+#include "UObject/ObjectPtr.h"
+#include "Templates/SharedPointer.h"
+
+class FWaveformEditorRenderData;
+class FWaveformEditorTransportCoordinator;
+class FWaveformEditorZoomController;
+class SWaveformTransformationRenderLayer;
+class UWaveformTransformationBase;
+
+class FWaveformTransformationRenderLayerFactory
+{
+public:
+	explicit FWaveformTransformationRenderLayerFactory(
+		TSharedRef<FWaveformEditorRenderData> InWaveformRenderData, 
+		TSharedRef<FWaveformEditorTransportCoordinator> InTransportCoordinator, 
+		TSharedRef<FWaveformEditorZoomController> InZoomController);
+
+	~FWaveformTransformationRenderLayerFactory() = default;
+
+	TSharedPtr<SWaveformTransformationRenderLayer> Create(TObjectPtr<UWaveformTransformationBase> InTransformationToRender) const;
+	TSharedPtr<SWaveformTransformationRenderLayer> CreateDurationHiglightLayer () const;
+
+private:
+	TSharedPtr<FWaveformEditorRenderData> WaveformRenderData = nullptr;
+	TSharedPtr<FWaveformEditorTransportCoordinator> TransportCoordinator = nullptr;
+	TSharedPtr<FWaveformEditorZoomController> ZoomController = nullptr;
+};
\ No newline at end of file
diff --git a/Engine/Plugins/Editor/WaveformEditor/Source/WaveformTransformationsWidgets/Public/WaveformTransformationsRenderManager.h b/Engine/Plugins/Editor/WaveformEditor/Source/WaveformTransformationsWidgets/Public/WaveformTransformationsRenderManager.h
new file mode 100644
--- /dev/null
+++ b/Engine/Plugins/Editor/WaveformEditor/Source/WaveformTransformationsWidgets/Public/WaveformTransformationsRenderManager.h
@@ -0,0 +1,85 @@
+// Copyright Epic Games, Inc. All Rights Reserved.
+
+#pragma once
+
+#include "Containers/ArrayView.h"
+#include "IWaveformTransformation.h"
+#include "WaveformTransformationRenderLayerFactory.h"
+
+class USoundWave;
+class UWaveformTransformationBase;
+class SWaveformTransformationRenderLayer;
+class SWidget;
+
+DECLARE_MULTICAST_DELEGATE_TwoParams(FOnLayersChainGenerated, TSharedPtr<SWidget>* /*First Layer ptr*/, const int32 /* NLayers */)
+DECLARE_MULTICAST_DELEGATE_SixParams(FOnRenderDataGenerated, const uint8* /* First RawPCMData element */, const uint32 /*Num Samples*/, const uint32 /*First Edited Sample*/, const uint32 /*Last Edited Sample*/, const uint32/*Sample Rate*/, const uint16 /*Num Channels*/)
+
+
+/*************************************************************************************/
+/* FWaveformTransformationsRenderManager											 */
+/* The Waveform Transformations render manager can produce UI to display a chain 	 */
+/* of waveform transformations														 */
+/*  																				 */
+/* The main UI elements created are:												 */
+/*  Render Data : a uint8 array containing the transformed samples of the waveform   */ 
+/*  Transform Layers : an array of Widgets containing the UI for each transformation */ 
+/*  in the chain																	 */
+/* 																					 */
+/* These are created so that the entire stack of transformation is displayable.		 */
+/* E.g : if a 10s long file is trimmed from second 2, seconds 0 to 1 will 			 */
+/* still be present in the render data.												 */
+/* UI widgets are passed a struct with information about the transformation to       */
+/* display property (e.g. StartFrameOffset, Sampleduration ecc)                      */
+/* 																					 */
+/* UIs for different transformations are registered withand spawned by 				 */
+/* FWaveformTransformationRenderLayerFactory.										 */
+/* Transformations don't necessarily have a widget UI. In that case they are 		 */
+/* only reflected in  the render data. 												 */
+/*************************************************************************************/
+
+
+class WAVEFORMTRANSFORMATIONSWIDGETS_API FWaveformTransformationsRenderManager
+{
+public:
+	explicit FWaveformTransformationsRenderManager(
+		TObjectPtr<USoundWave> InSoundWave, 
+		TSharedRef<FWaveformEditorRenderData> InWaveformRenderData, 
+		TSharedRef<FWaveformEditorTransportCoordinator> InTransportCoordinator, 
+		TSharedRef<FWaveformEditorZoomController> InZoomController
+	);
+
+	/** Used to generate the stack of transformations UI widgets			*/
+	/** Should be called when the waveform transformation chain is changed	*/
+	void GenerateLayersChain();
+
+	/** Used to generate updated render data and pass transformation info to the widgets */
+	/** Should be called when the transformations parameters are changed				 */
+	void UpdateRenderElements();
+
+	TArrayView<TSharedPtr<SWidget>> GetTransformLayers() const;
+
+	/** Called when a new layer chain of transformations UI is generated */
+	FOnLayersChainGenerated OnLayersChainGenerated;
+
+	/** Called when new render data is generated */
+	FOnRenderDataGenerated OnRenderDataGenerated;
+
+private:
+	void GenerateRenderDataInternal();
+
+	void CreateDurationHighlightLayer();
+
+	TArray<Audio::FTransformationPtr> CreateTransformations() const;
+	const bool CanChainChangeFileLength(const TArray<Audio::FTransformationPtr>& TransformationArray) const;
+
+	TArray<TObjectPtr<UWaveformTransformationBase>> TransformationsToRender;
+	TArray<TSharedPtr<SWaveformTransformationRenderLayer>> RenderLayers;
+	TArrayView<TSharedPtr<SWidget>> RenderLayersWidgetView;
+
+	TObjectPtr<USoundWave> SoundWaveToRender = nullptr;
+
+	TArray<uint8> RawPCMData;
+
+	TUniquePtr<FWaveformTransformationRenderLayerFactory> LayersFactory = nullptr;
+	TSharedPtr<SWaveformTransformationRenderLayer> DurationHiglightLayer = nullptr;
+};
\ No newline at end of file
diff --git a/Engine/Plugins/Editor/WaveformEditor/Source/WaveformTransformationsWidgets/Public/WaveformTransformationsWidgetsModule.h b/Engine/Plugins/Editor/WaveformEditor/Source/WaveformTransformationsWidgets/Public/WaveformTransformationsWidgetsModule.h
new file mode 100644
--- /dev/null
+++ b/Engine/Plugins/Editor/WaveformEditor/Source/WaveformTransformationsWidgets/Public/WaveformTransformationsWidgetsModule.h
@@ -0,0 +1,9 @@
+// Copyright Epic Games, Inc. All Rights Reserved.
+
+#pragma once
+
+#include "Modules/ModuleInterface.h"
+
+class FWaveformTransformationsWidgetsModule : public IModuleInterface
+{
+};
\ No newline at end of file
diff --git a/Engine/Plugins/Editor/WaveformEditor/Source/WaveformTransformationsWidgets/WaveformTransformationsWidgets.build.cs b/Engine/Plugins/Editor/WaveformEditor/Source/WaveformTransformationsWidgets/WaveformTransformationsWidgets.build.cs
new file mode 100644
--- /dev/null
+++ b/Engine/Plugins/Editor/WaveformEditor/Source/WaveformTransformationsWidgets/WaveformTransformationsWidgets.build.cs
@@ -0,0 +1,33 @@
+// Copyright Epic Games, Inc. All Rights Reserved.
+
+using UnrealBuildTool;
+
+namespace UnrealBuildTool.Rules
+{
+	public class WaveformTransformationsWidgets : ModuleRules
+	{
+		public WaveformTransformationsWidgets(ReadOnlyTargetRules Target) : base(Target)
+		{
+			PCHUsage = PCHUsageMode.UseExplicitOrSharedPCHs;
+
+			PublicDependencyModuleNames.AddRange(
+				new string[] {
+					"AudioExtensions",
+					"ContentBrowser",
+					"Core",
+					"CoreUObject",
+					"DeveloperSettings",
+					"EditorStyle",
+					"Engine",
+					"InputCore",
+					"Slate",
+					"SignalProcessing",
+					"SlateCore", 
+					"ToolMenus", 
+					"WaveformTransformations",
+					"WaveformEditorWidgets",
+				}
+			);
+		}
+	}
+}
\ No newline at end of file
diff --git a/Engine/Plugins/Editor/WaveformEditor/WaveformEditor.uplugin b/Engine/Plugins/Editor/WaveformEditor/WaveformEditor.uplugin
new file mode 100644
--- /dev/null
+++ b/Engine/Plugins/Editor/WaveformEditor/WaveformEditor.uplugin
@@ -0,0 +1,45 @@
+{
+	"FileVersion": 3,
+	"Version": 0,
+	"VersionName": "0.0",
+	"FriendlyName": "Waveform Editor",
+	"Description": "Editor tool for waveforms",
+	"Category": "Audio",
+	"CreatedBy": "Epic Games, Inc.",
+	"CreatedByURL": "http://epicgames.com",
+	"DocsURL": "",
+	"MarketplaceURL": "",
+	"SupportURL": "",
+	"EnabledByDefault": false,
+	"CanContainContent": true,
+	"IsBetaVersion": true,
+	"Installed": false,
+	"Modules": [
+		{
+			"Name": "WaveformEditor",
+			"Type": "Editor",
+			"LoadingPhase": "Default"
+		},
+		{
+			"Name": "WaveformEditorWidgets",
+			"Type": "Editor",
+			"LoadingPhase": "Default"
+		},
+		{
+			"Name": "WaveformTransformations",
+			"Type": "UncookedOnly",
+			"LoadingPhase": "PreDefault"
+		},
+		{
+			"Name": "WaveformTransformationsWidgets",
+			"Type": "Editor",
+			"LoadingPhase": "Default"
+		}
+	],
+	"Plugins": [
+		{
+			"Name": "AudioSynesthesia",
+			"Enabled": true
+		}
+	]
+}
\ No newline at end of file
