[WaveformEditor] fix some static compile warnings
#trivial

diff --git a/Engine/Plugins/Editor/WaveformEditor/Source/WaveformEditor/Private/WaveformEditor.cpp b/Engine/Plugins/Editor/WaveformEditor/Source/WaveformEditor/Private/WaveformEditor.cpp
--- a/Engine/Plugins/Editor/WaveformEditor/Source/WaveformEditor/Private/WaveformEditor.cpp
+++ b/Engine/Plugins/Editor/WaveformEditor/Source/WaveformEditor/Private/WaveformEditor.cpp
@@ -84,6 +84,11 @@ bool FWaveformEditor::SetupAudioComponent()
 				AudioComponent = FAudioDevice::CreateComponent(SoundBase);
 			}
 		}
+
+		if (AudioComponent == nullptr)
+		{
+			return false;
+		}
 	}
 
 	AudioComponent->bAutoDestroy = false;
@@ -93,7 +98,7 @@ bool FWaveformEditor::SetupAudioComponent()
 	AudioComponent->bCenterChannelOnly = false;
 	AudioComponent->bIsPreviewSound = true;
 
-	return AudioComponent != nullptr;
+	return true;
 }
 
 bool FWaveformEditor::SetUpTransportController()
diff --git a/Engine/Plugins/Editor/WaveformEditor/Source/WaveformEditor/Private/WaveformEditorInstantiator.cpp b/Engine/Plugins/Editor/WaveformEditor/Source/WaveformEditor/Private/WaveformEditorInstantiator.cpp
--- a/Engine/Plugins/Editor/WaveformEditor/Source/WaveformEditor/Private/WaveformEditorInstantiator.cpp
+++ b/Engine/Plugins/Editor/WaveformEditor/Source/WaveformEditor/Private/WaveformEditorInstantiator.cpp
@@ -64,10 +64,11 @@ bool FWaveformEditorInstantiator::CanSoundWaveBeOpenedInEditor(const USoundWave*
 {
 	if (SoundWaveToEdit == nullptr)
 	{
-		UE_LOG(LogWaveformEditor, Warning, TEXT("Could not open waveform editor for soundwave %s, soundwave is null"), *(SoundWaveToEdit->GetName()))
+		UE_LOG(LogWaveformEditor, Warning, TEXT("Could not open waveform editor from null SoundWave"))
 		return false;
 	}
-	else if (SoundWaveToEdit->GetDuration() == 0.f)
+	
+	if (SoundWaveToEdit->GetDuration() == 0.f)
 	{
 		UE_LOG(LogWaveformEditor, Warning, TEXT("Could not open waveform editor for soundwave %s, duration is 0"), *(SoundWaveToEdit->GetName()))
 		return false;
