[WaveformEditor] Undef Loctext namespace in time ruler
#trivial

diff --git a/Engine/Plugins/Editor/WaveformEditor/Source/WaveformEditorWidgets/Private/SWaveformEditorTimeRuler.cpp b/Engine/Plugins/Editor/WaveformEditor/Source/WaveformEditorWidgets/Private/SWaveformEditorTimeRuler.cpp
--- a/Engine/Plugins/Editor/WaveformEditor/Source/WaveformEditorWidgets/Private/SWaveformEditorTimeRuler.cpp
+++ b/Engine/Plugins/Editor/WaveformEditor/Source/WaveformEditorWidgets/Private/SWaveformEditorTimeRuler.cpp
@@ -254,3 +254,5 @@ FReply SWaveformEditorTimeRuler::OnMouseMove(const FGeometry& MyGeometry, const
 {
 	return TransportCoordinator->ReceiveMouseMove(*this, MyGeometry, MouseEvent);
 }
+
+#undef LOCTEXT_NAMESPACE
\ No newline at end of file
