[WaveformEditor] fix for V593 warnings
#trivial

diff --git a/Engine/Plugins/Editor/WaveformEditor/Source/WaveformEditor/Private/WaveformEditorWaveWriter.cpp b/Engine/Plugins/Editor/WaveformEditor/Source/WaveformEditor/Private/WaveformEditorWaveWriter.cpp
--- a/Engine/Plugins/Editor/WaveformEditor/Source/WaveformEditor/Private/WaveformEditorWaveWriter.cpp
+++ b/Engine/Plugins/Editor/WaveformEditor/Source/WaveformEditor/Private/WaveformEditorWaveWriter.cpp
@@ -91,7 +91,9 @@ Audio::TSampleBuffer<> FWaveformEditorWaveWriter::GenerateSampleBuffer() const
 			Transformation->ProcessAudio(TransformationInfo);
 		}
 
-		if (const float MaxValue = Audio::ArrayMaxAbsValue(Buffer) > 1.f)
+		const float MaxValue = Audio::ArrayMaxAbsValue(Buffer);
+
+		if (MaxValue > 1.f)
 		{
 			Audio::ArrayMultiplyByConstantInPlace(Buffer, 1.f / MaxValue);
 		}
diff --git a/Engine/Plugins/Editor/WaveformEditor/Source/WaveformTransformationsWidgets/Private/WaveformTransformationsRenderManager.cpp b/Engine/Plugins/Editor/WaveformEditor/Source/WaveformTransformationsWidgets/Private/WaveformTransformationsRenderManager.cpp
--- a/Engine/Plugins/Editor/WaveformEditor/Source/WaveformTransformationsWidgets/Private/WaveformTransformationsRenderManager.cpp
+++ b/Engine/Plugins/Editor/WaveformEditor/Source/WaveformTransformationsWidgets/Private/WaveformTransformationsRenderManager.cpp
@@ -124,8 +124,9 @@ void FWaveformTransformationsRenderManager::GenerateRenderDataInternal()
 			OutputBuffer = TransformationsBuffer;
 		}
 		
+		const float MaxValue = Audio::ArrayMaxAbsValue(OutputBuffer);
 
-		if (const float MaxValue = Audio::ArrayMaxAbsValue(OutputBuffer) > 1.f)
+		if (MaxValue > 1.f)
 		{
 			Audio::ArrayMultiplyByConstantInPlace(OutputBuffer, 1.f / MaxValue);
 		}
