[WaveformEditor] add support for undo
fix for #jira UE-161401
#trivial

diff --git a/Engine/Plugins/Editor/WaveformEditor/Source/WaveformEditor/Private/WaveformEditor.cpp b/Engine/Plugins/Editor/WaveformEditor/Source/WaveformEditor/Private/WaveformEditor.cpp
--- a/Engine/Plugins/Editor/WaveformEditor/Source/WaveformEditor/Private/WaveformEditor.cpp
+++ b/Engine/Plugins/Editor/WaveformEditor/Source/WaveformEditor/Private/WaveformEditor.cpp
@@ -4,6 +4,7 @@
 
 #include "AudioDevice.h"
 #include "Components/AudioComponent.h"
+#include "Misc/TransactionObjectEvent.h"
 #include "Sound/SoundWave.h"
 #include "Styling/AppStyle.h"
 #include "SWaveformPanel.h"
@@ -13,8 +14,8 @@
 #include "WaveformEditorLog.h"
 #include "WaveformEditorRenderData.h"
 #include "WaveformEditorStyle.h"
-#include "WaveformTransformationsRenderManager.h"
 #include "WaveformEditorWaveWriter.h"
+#include "WaveformTransformationsRenderManager.h"
 #include "Widgets/Docking/SDockTab.h"
 
 #define LOCTEXT_NAMESPACE "WaveformEditor"
@@ -45,6 +46,8 @@ bool FWaveformEditor::Init(const EToolkitMode::Type Mode, const TSharedPtr<ITool
 
 	const TSharedRef<FTabManager::FLayout>  StandaloneDefaultLayout = SetupStandaloneLayout();
 
+	GEditor->RegisterForUndo(this);
+
 	if (bIsInitialized)
 	{
 		const bool bCreateDefaultStandaloneMenu = true;
@@ -342,6 +345,36 @@ void FWaveformEditor::NotifyPostChange(const FPropertyChangedEvent& PropertyChan
 	TransformationsRenderManager->UpdateRenderElements();
 }
 
+void FWaveformEditor::PostUndo(bool bSuccess)
+{
+	if (bSuccess)
+	{
+		TransformationsRenderManager->GenerateLayersChain();
+		TransformationsRenderManager->UpdateRenderElements();
+	}
+}
+
+void FWaveformEditor::PostRedo(bool bSuccess)
+{
+	PostUndo(bSuccess);
+}
+
+bool FWaveformEditor::MatchesContext(const FTransactionContext& InContext, const TArray<TPair<UObject*, FTransactionObjectEvent>>& TransactionObjectContexts) const
+{
+	bool bShoouldUndo = false;
+	for (const TPair<UObject*, FTransactionObjectEvent>& TransactionObjectPair : TransactionObjectContexts)
+	{
+		UObject* Object = TransactionObjectPair.Key;
+		const UClass* ObjectClass = Object->GetClass();
+		if (ObjectClass && ObjectClass->IsChildOf(UWaveformTransformationBase::StaticClass()))
+		{
+			bShoouldUndo = true;
+		}
+	}
+
+	return bShoouldUndo;
+}
+
 bool FWaveformEditor::SetUpPropertiesView()
 {
 	if (SoundWave == nullptr)
diff --git a/Engine/Plugins/Editor/WaveformEditor/Source/WaveformEditor/Public/WaveformEditor.h b/Engine/Plugins/Editor/WaveformEditor/Source/WaveformEditor/Public/WaveformEditor.h
--- a/Engine/Plugins/Editor/WaveformEditor/Source/WaveformEditor/Public/WaveformEditor.h
+++ b/Engine/Plugins/Editor/WaveformEditor/Source/WaveformEditor/Public/WaveformEditor.h
@@ -2,6 +2,7 @@
 
 #pragma once
 
+#include "EditorUndoClient.h"
 #include "Misc/NotifyHook.h"
 #include "Toolkits/AssetEditorToolkit.h"
 #include "WaveformEditorTransportController.h"
@@ -13,7 +14,12 @@ class SDockTab;
 class USoundWave;
 class UAudioComponent;
 
-class WAVEFORMEDITOR_API FWaveformEditor : public FAssetEditorToolkit, public FGCObject, public FNotifyHook
+class WAVEFORMEDITOR_API FWaveformEditor 
+	: public FAssetEditorToolkit
+	, public FEditorUndoClient
+	, public FGCObject
+	, public FNotifyHook
+
 {
 public:
 	bool Init(const EToolkitMode::Type Mode, const TSharedPtr<IToolkitHost>& InitToolkitHost, USoundWave* SoundWaveToEdit);
@@ -27,9 +33,15 @@ public:
 	virtual FString GetWorldCentricTabPrefix() const override;
 	virtual FLinearColor GetWorldCentricTabColorScale() const override;
 
+	/** FNotifyHook interface */
 	void NotifyPreChange(class FEditPropertyChain* PropertyAboutToChange) override {};
 	void NotifyPostChange(const FPropertyChangedEvent& PropertyChangedEvent, class FEditPropertyChain* PropertyThatChanged) override;
 
+	/** FEditorUndo interface */
+	void PostUndo(bool bSuccess) override;
+	void PostRedo(bool bSuccess) override;
+	bool MatchesContext(const FTransactionContext& InContext, const TArray<TPair<UObject*, FTransactionObjectEvent>>& TransactionObjectContexts) const override;
+
 private:
 	bool SetupAudioComponent();
 	bool SetUpTransportController();
