[WaveformEditor] Add Transformations Only Tab and show it in default layout #jira UE-161208 #preflight 63074757e94540ecae632c59

diff --git a/Engine/Plugins/Editor/WaveformEditor/Source/WaveformEditor/Private/WaveformEditor.cpp b/Engine/Plugins/Editor/WaveformEditor/Source/WaveformEditor/Private/WaveformEditor.cpp
--- a/Engine/Plugins/Editor/WaveformEditor/Source/WaveformEditor/Private/WaveformEditor.cpp
+++ b/Engine/Plugins/Editor/WaveformEditor/Source/WaveformEditor/Private/WaveformEditor.cpp
@@ -22,6 +22,7 @@
 
 const FName FWaveformEditor::AppIdentifier("WaveformEditorApp");
 const FName FWaveformEditor::PropertiesTabId("WaveformEditor_Properties");
+const FName FWaveformEditor::TransformationsTabId("WaveformEditor_Transformations");
 const FName FWaveformEditor::WaveformDisplayTabId("WaveformEditor_Display");
 const FName FWaveformEditor::EditorName("Waveform Editor");
 const FName FWaveformEditor::ToolkitFName("WaveformEditor");
@@ -34,7 +35,7 @@ bool FWaveformEditor::Init(const EToolkitMode::Type Mode, const TSharedPtr<ITool
 
 	bool bIsInitialized = true;
 	
-	bIsInitialized &= SetUpPropertiesView();
+	bIsInitialized &= SetUpDetailsViews();
 	bIsInitialized &= SetUpWaveformPanel();
 	bIsInitialized &= SetupAudioComponent();
 	bIsInitialized &= SetUpTransportController();
@@ -147,6 +148,11 @@ void FWaveformEditor::RegisterTabSpawners(const TSharedRef<FTabManager>& InTabMa
 		.SetGroup(WorkspaceMenuCategory.ToSharedRef())
 		.SetIcon(FSlateIcon(FAppStyle::GetAppStyleSetName(), "LevelEditor.Tabs.Details"));
 
+	InTabManager->RegisterTabSpawner(TransformationsTabId, FOnSpawnTab::CreateSP(this, &FWaveformEditor::SpawnTab_Transformations))
+		.SetDisplayName(LOCTEXT("TransformationsTab", "Transformations"))
+		.SetGroup(WorkspaceMenuCategory.ToSharedRef())
+		.SetIcon(FSlateIcon(FAppStyle::GetAppStyleSetName(), "LevelEditor.Tabs.Details"));
+
 	InTabManager->RegisterTabSpawner(WaveformDisplayTabId, FOnSpawnTab::CreateSP(this, &FWaveformEditor::SpawnTab_WaveformDisplay))
 		.SetDisplayName(LOCTEXT("WaveformDisplayTab", "WaveformDisplay"))
 		.SetGroup(WorkspaceMenuCategory.ToSharedRef())
@@ -375,7 +381,7 @@ bool FWaveformEditor::MatchesContext(const FTransactionContext& InContext, const
 	return bShoouldUndo;
 }
 
-bool FWaveformEditor::SetUpPropertiesView()
+bool FWaveformEditor::SetUpDetailsViews()
 {
 	if (SoundWave == nullptr)
 	{
@@ -383,13 +389,18 @@ bool FWaveformEditor::SetUpPropertiesView()
 		return false;
 	}
 
+	FPropertyEditorModule& PropertyModule = FModuleManager::LoadModuleChecked<FPropertyEditorModule>("PropertyEditor");
 	FDetailsViewArgs Args;
 	Args.bHideSelectionTip = true;
 	Args.NotifyHook = this;
 
-	FPropertyEditorModule& PropertyModule = FModuleManager::LoadModuleChecked<FPropertyEditorModule>("PropertyEditor");
-	PropertiesView = PropertyModule.CreateDetailView(Args);
-	PropertiesView->SetObject(SoundWave);
+	PropertiesDetails = PropertyModule.CreateDetailView(Args);
+	PropertiesDetails->SetObject(SoundWave);
+
+	TransformationsDetails = PropertyModule.CreateDetailView(Args);
+	TransformationsView = TStrongObjectPtr(NewObject<UWaveformTransformationsViewHelper>());
+	TransformationsView->SetSoundWave(SoundWave);
+	TransformationsDetails->SetObject(TransformationsView.Get());
 
 	return true;
 }
@@ -407,7 +418,7 @@ TSharedRef<SDockTab> FWaveformEditor::SpawnTab_WaveformDisplay(const FSpawnTabAr
 
 const TSharedRef<FTabManager::FLayout> FWaveformEditor::SetupStandaloneLayout()
 {
-	const TSharedRef<FTabManager::FLayout> StandaloneDefaultLayout = FTabManager::NewLayout("Standalone_WaveformEditor_v1")
+	const TSharedRef<FTabManager::FLayout> StandaloneDefaultLayout = FTabManager::NewLayout("Standalone_WaveformEditor_v2")
 		->AddArea
 		(
 			FTabManager::NewPrimaryArea()->SetOrientation(Orient_Vertical)
@@ -418,8 +429,9 @@ const TSharedRef<FTabManager::FLayout> FWaveformEditor::SetupStandaloneLayout()
 				(
 					FTabManager::NewStack()
 					->SetSizeCoefficient(0.2f)
+					->AddTab(TransformationsTabId, ETabState::OpenedTab)
 					->AddTab(PropertiesTabId, ETabState::OpenedTab)
-					->SetForegroundTab(PropertiesTabId)
+					->SetForegroundTab(TransformationsTabId)
 
 				)
 				->Split
@@ -444,7 +456,18 @@ TSharedRef<SDockTab> FWaveformEditor::SpawnTab_Properties(const FSpawnTabArgs& A
 	return SNew(SDockTab)
 		.Label(LOCTEXT("SoundWaveDetailsTitle", "Details"))
 		[
-			PropertiesView.ToSharedRef()
+			PropertiesDetails.ToSharedRef()
+		];
+}
+
+TSharedRef<SDockTab> FWaveformEditor::SpawnTab_Transformations(const FSpawnTabArgs& Args)
+{
+	check(Args.GetTabId() == TransformationsTabId);
+
+	return SNew(SDockTab)
+		.Label(LOCTEXT("SoundWaveTransformationsTitle", "Transformations"))
+		[
+			TransformationsDetails.ToSharedRef()
 		];
 }
 
diff --git a/Engine/Plugins/Editor/WaveformEditor/Source/WaveformEditor/Private/WaveformEditorCustomDetailsHelpers.cpp b/Engine/Plugins/Editor/WaveformEditor/Source/WaveformEditor/Private/WaveformEditorCustomDetailsHelpers.cpp
new file mode 100644
--- /dev/null
+++ b/Engine/Plugins/Editor/WaveformEditor/Source/WaveformEditor/Private/WaveformEditorCustomDetailsHelpers.cpp
@@ -0,0 +1,13 @@
+// Copyright Epic Games, Inc. All Rights Reserved.
+
+#include "WaveformEditorCustomDetailsHelpers.h"
+
+void UWaveformTransformationsViewHelper::SetSoundWave(UObject* InSoundWave)
+{
+	SoundWave = InSoundWave;
+}
+
+const TWeakObjectPtr<UObject> UWaveformTransformationsViewHelper::GetSoundWave()
+{
+	return SoundWave;
+}
diff --git a/Engine/Plugins/Editor/WaveformEditor/Source/WaveformEditor/Private/WaveformEditorDetailsCustomization.cpp b/Engine/Plugins/Editor/WaveformEditor/Source/WaveformEditor/Private/WaveformEditorDetailsCustomization.cpp
new file mode 100644
--- /dev/null
+++ b/Engine/Plugins/Editor/WaveformEditor/Source/WaveformEditor/Private/WaveformEditorDetailsCustomization.cpp
@@ -0,0 +1,28 @@
+// Copyright Epic Games, Inc. All Rights Reserved.
+
+#include "WaveformEditorDetailsCustomization.h"
+
+#include "DetailCategoryBuilder.h"
+#include "DetailLayoutBuilder.h"
+#include "Sound/SoundWave.h"
+#include "WaveformEditorCustomDetailsHelpers.h"
+
+void FWaveformTransformationsDetailsCustomization::CustomizeDetails(IDetailLayoutBuilder& DetailLayout)
+{	
+	TArray<TWeakObjectPtr<UObject>> Objects;
+	DetailLayout.GetObjectsBeingCustomized(Objects);
+
+	//there should only one soundwave per editor instance open
+	check(Objects.Num() == 1)
+
+	if (UWaveformTransformationsViewHelper* TransformationsView = CastChecked<UWaveformTransformationsViewHelper>(Objects.Last()))
+	{
+		SoundWaveObject = TransformationsView->GetSoundWave();
+	}
+
+	IDetailCategoryBuilder& CategoryBuilder = DetailLayout.EditCategory("General");
+
+	TArray<UObject*> ObjArray{ SoundWaveObject.Get() };
+	CategoryBuilder.AddExternalObjectProperty(ObjArray, GET_MEMBER_NAME_CHECKED(USoundWave, Transformations));
+	CategoryBuilder.InitiallyCollapsed(false);
+}
\ No newline at end of file
diff --git a/Engine/Plugins/Editor/WaveformEditor/Source/WaveformEditor/Private/WaveformEditorDetailsCustomization.h b/Engine/Plugins/Editor/WaveformEditor/Source/WaveformEditor/Private/WaveformEditorDetailsCustomization.h
new file mode 100644
--- /dev/null
+++ b/Engine/Plugins/Editor/WaveformEditor/Source/WaveformEditor/Private/WaveformEditorDetailsCustomization.h
@@ -0,0 +1,17 @@
+// Copyright Epic Games, Inc. All Rights Reserved.
+
+#pragma once
+
+#include "IDetailCustomization.h"
+#include "UObject/WeakObjectPtr.h"
+
+class IDetailLayoutBuilder;
+
+class FWaveformTransformationsDetailsCustomization : public IDetailCustomization
+{
+public:	
+	void CustomizeDetails(IDetailLayoutBuilder& DetailLayout) override;
+
+private:
+	TWeakObjectPtr<UObject> SoundWaveObject;
+};
\ No newline at end of file
diff --git a/Engine/Plugins/Editor/WaveformEditor/Source/WaveformEditor/Private/WaveformEditorModule.cpp b/Engine/Plugins/Editor/WaveformEditor/Source/WaveformEditor/Private/WaveformEditorModule.cpp
--- a/Engine/Plugins/Editor/WaveformEditor/Source/WaveformEditor/Private/WaveformEditorModule.cpp
+++ b/Engine/Plugins/Editor/WaveformEditor/Source/WaveformEditor/Private/WaveformEditorModule.cpp
@@ -3,18 +3,28 @@
 #include "WaveformEditorModule.h"
 
 #include "Modules/ModuleManager.h"
+#include "PropertyEditorModule.h"
 #include "WaveformEditorCommands.h"
+#include "WaveformEditorCustomDetailsHelpers.h"
+#include "WaveformEditorDetailsCustomization.h"
 #include "WaveformEditorInstantiator.h"
 #include "WaveformEditorLog.h"
 
 DEFINE_LOG_CATEGORY(LogWaveformEditor);
 
+
 void FWaveformEditorModule::StartupModule()
 {
 	FWaveformEditorCommands::Register();
 
 	WaveformEditorInstantiator = MakeShared<FWaveformEditorInstantiator>();
 	RegisterContentBrowserExtensions(WaveformEditorInstantiator.Get());
+
+	FPropertyEditorModule& PropertyModule = FModuleManager::LoadModuleChecked<FPropertyEditorModule>("PropertyEditor");
+
+	PropertyModule.RegisterCustomClassLayout(
+		UWaveformTransformationsViewHelper::StaticClass()->GetFName(),
+		FOnGetDetailCustomizationInstance::CreateLambda([]() { return MakeShared<FWaveformTransformationsDetailsCustomization>(); }));
 }
 
 void FWaveformEditorModule::ShutdownModule()
diff --git a/Engine/Plugins/Editor/WaveformEditor/Source/WaveformEditor/Public/WaveformEditor.h b/Engine/Plugins/Editor/WaveformEditor/Source/WaveformEditor/Public/WaveformEditor.h
--- a/Engine/Plugins/Editor/WaveformEditor/Source/WaveformEditor/Public/WaveformEditor.h
+++ b/Engine/Plugins/Editor/WaveformEditor/Source/WaveformEditor/Public/WaveformEditor.h
@@ -7,6 +7,7 @@
 #include "Toolkits/AssetEditorToolkit.h"
 #include "WaveformEditorTransportController.h"
 #include "WaveformEditorTransportCoordinator.h"
+#include "WaveformEditorCustomDetailsHelpers.h"
 #include "WaveformEditorZoomController.h"
 
 class IToolkitHost;
@@ -56,9 +57,11 @@ private:
 	bool RegisterToolbar();
 	bool BindCommands();
 
-	/**	Properties tab set up */
+	/**	Details tabs set up */
 	TSharedRef<SDockTab> SpawnTab_Properties(const FSpawnTabArgs& Args);
-	bool SetUpPropertiesView();
+	TSharedRef<SDockTab> SpawnTab_Transformations(const FSpawnTabArgs& Args);
+
+	bool SetUpDetailsViews();
 
 	/**	Waveform view tab setup */
 	TSharedRef<SDockTab> SpawnTab_WaveformDisplay(const FSpawnTabArgs& Args);
@@ -97,13 +100,18 @@ private:
 	TSharedPtr<FWaveformEditorZoomController> ZoomManager = nullptr;
 
 	/** Properties tab */
-	TSharedPtr<IDetailsView> PropertiesView;
+	TSharedPtr<IDetailsView> PropertiesDetails;
+
+	/** Transformations tab */
+	TSharedPtr<IDetailsView> TransformationsDetails;
+	TStrongObjectPtr<UWaveformTransformationsViewHelper> TransformationsView;
 
 	/** Settings Editor App Identifier */
 	static const FName AppIdentifier;
 
 	/** Tab Ids */
 	static const FName PropertiesTabId;
+	static const FName TransformationsTabId;
 	static const FName WaveformDisplayTabId;
 
 	static const FName EditorName;
diff --git a/Engine/Plugins/Editor/WaveformEditor/Source/WaveformEditor/Public/WaveformEditorCustomDetailsHelpers.h b/Engine/Plugins/Editor/WaveformEditor/Source/WaveformEditor/Public/WaveformEditorCustomDetailsHelpers.h
new file mode 100644
--- /dev/null
+++ b/Engine/Plugins/Editor/WaveformEditor/Source/WaveformEditor/Public/WaveformEditorCustomDetailsHelpers.h
@@ -0,0 +1,22 @@
+// Copyright Epic Games, Inc. All Rights Reserved.
+
+#pragma once
+
+#include "UObject/Object.h"
+
+#include "WaveformEditorCustomDetailsHelpers.generated.h"
+
+// Simple helper classes for custom interfaces details tab 
+
+UCLASS(MinimalAPI)
+class UWaveformTransformationsViewHelper : public UObject
+{
+	GENERATED_BODY()
+public:
+	UWaveformTransformationsViewHelper() = default;
+
+	void SetSoundWave(UObject* InSoundWave);
+	const TWeakObjectPtr<UObject> GetSoundWave();
+private:
+	TWeakObjectPtr<UObject> SoundWave;
+};
\ No newline at end of file
