[WaveformEditor] Fix typo in IMPLEMENT_MODULE macro for waveform transformations widgets #preflight 63087156d65521a5a59eb4a4 #trivial #rb none

diff --git a/Engine/Plugins/Editor/WaveformEditor/Source/WaveformTransformationsWidgets/Private/WaveformTransformationsWidgetsModule.cpp b/Engine/Plugins/Editor/WaveformEditor/Source/WaveformTransformationsWidgets/Private/WaveformTransformationsWidgetsModule.cpp
--- a/Engine/Plugins/Editor/WaveformEditor/Source/WaveformTransformationsWidgets/Private/WaveformTransformationsWidgetsModule.cpp
+++ b/Engine/Plugins/Editor/WaveformEditor/Source/WaveformTransformationsWidgets/Private/WaveformTransformationsWidgetsModule.cpp
@@ -3,4 +3,4 @@
 #include "WaveformTransformationsWidgetsModule.h"
 #include "Modules/ModuleManager.h"
 
-IMPLEMENT_MODULE(FWaveformTransformationsWidgetsModule, WaveformTransformations);
\ No newline at end of file
+IMPLEMENT_MODULE(FWaveformTransformationsWidgetsModule, WaveformTransformationsWidgets);
\ No newline at end of file
