[WaveformEditor] Decouple Waveform Transformations UIs from SWidget and add support for user input #preflight 630dd473660db81edb8ac834

diff --git a/Engine/Plugins/Editor/WaveformEditor/Source/WaveformEditor/Private/WaveformEditor.cpp b/Engine/Plugins/Editor/WaveformEditor/Source/WaveformEditor/Private/WaveformEditor.cpp
--- a/Engine/Plugins/Editor/WaveformEditor/Source/WaveformEditor/Private/WaveformEditor.cpp
+++ b/Engine/Plugins/Editor/WaveformEditor/Source/WaveformEditor/Private/WaveformEditor.cpp
@@ -447,8 +447,6 @@ const TSharedRef<FTabManager::FLayout> FWaveformEditor::SetupStandaloneLayout()
 	return StandaloneDefaultLayout;
 }
 
-
-
 TSharedRef<SDockTab> FWaveformEditor::SpawnTab_Properties(const FSpawnTabArgs& Args)
 {
 	check(Args.GetTabId() == PropertiesTabId);
@@ -470,7 +468,6 @@ TSharedRef<SDockTab> FWaveformEditor::SpawnTab_Transformations(const FSpawnTabAr
 			TransformationsDetails.ToSharedRef()
 		];
 }
-
 bool FWaveformEditor::SetUpWaveformPanel()
 {
 	if (SoundWave == nullptr)
@@ -490,19 +487,24 @@ bool FWaveformEditor::SetUpWaveformPanel()
 	TransportCoordinator = MakeShared<FWaveformEditorTransportCoordinator>(RenderData.ToSharedRef());
 	RenderData->OnRenderDataUpdated.AddSP(TransportCoordinator.Get(), &FWaveformEditorTransportCoordinator::HandleRenderDataUpdate);
 
-	TransformationsRenderManager = MakeShared<FWaveformTransformationsRenderManager>(SoundWave, RenderData.ToSharedRef(), TransportCoordinator.ToSharedRef(), ZoomManager.ToSharedRef());
+	TFunction<void(FPropertyChangedEvent&, FEditPropertyChain*)> PropertyChangeNotifier = [this](FPropertyChangedEvent& PropertyChangedEvent, FEditPropertyChain* PropertyThatChanged) {
+		this->NotifyPostChange(PropertyChangedEvent, PropertyThatChanged); 
+	};
+
+	TransformationsRenderManager = MakeShared<FWaveformTransformationsRenderManager>(SoundWave, RenderData.ToSharedRef(), PropertyChangeNotifier);
 	TransformationsRenderManager->OnRenderDataGenerated.AddSP(RenderData.Get(), &FWaveformEditorRenderData::UpdateRenderData);
 
-	TSharedPtr<SWaveformTransformationsOverlay> TransformationsOverlay = SNew(SWaveformTransformationsOverlay, TransformationsRenderManager->GetTransformLayers());
+	TSharedPtr<SWaveformTransformationsOverlay> TransformationsOverlay = SNew(SWaveformTransformationsOverlay, TransformationsRenderManager->GetTransformLayers(), TransportCoordinator.ToSharedRef());
 
-	TransformationsRenderManager->OnLayersChainGenerated.AddSP(TransformationsOverlay.Get(), &SWaveformTransformationsOverlay::OnLayerChainUpdate);
+	TransformationsRenderManager->OnLayersChainGenerated.AddSP(TransformationsOverlay.Get(), &SWaveformTransformationsOverlay::OnLayerChainGenerated);
+	TransformationsRenderManager->OnRenderElementsUpdated.AddSP(TransformationsOverlay.Get(), &SWaveformTransformationsOverlay::UpdateLayerConstraints);
+	TransportCoordinator->OnDisplayRangeUpdated.AddSP(TransformationsOverlay.Get(), &SWaveformTransformationsOverlay::OnNewWaveformDisplayRange);
 	TransformationsRenderManager->UpdateRenderElements();
 
 	WaveformPanel = SNew(SWaveformPanel, RenderData.ToSharedRef(), TransportCoordinator.ToSharedRef(), ZoomManager.ToSharedRef(), TransformationsOverlay);
 
 	return WaveformPanel != nullptr;
 }
-
 void FWaveformEditor::HandlePlaybackPercentageChange(const UAudioComponent* InComponent, const USoundWave* InSoundWave, const float InPlaybackPercentage)
 {
 	const bool bIsStopped = AudioComponent->GetPlayState() == EAudioComponentPlayState::Stopped;
diff --git a/Engine/Plugins/Editor/WaveformEditor/Source/WaveformEditorWidgets/Private/SWaveformPanel.cpp b/Engine/Plugins/Editor/WaveformEditor/Source/WaveformEditorWidgets/Private/SWaveformPanel.cpp
--- a/Engine/Plugins/Editor/WaveformEditor/Source/WaveformEditorWidgets/Private/SWaveformPanel.cpp
+++ b/Engine/Plugins/Editor/WaveformEditor/Source/WaveformEditorWidgets/Private/SWaveformPanel.cpp
@@ -80,7 +80,7 @@ void SWaveformPanel::SetUpTimeRuler(TSharedRef<FWaveformEditorRenderData> InRend
 
 void SWaveformPanel::SetUpWaveformViewerOverlay(TSharedRef<FWaveformEditorTransportCoordinator> InTransportCoordinator, TSharedRef<FWaveformEditorZoomController> InZoomManager)
 {
-	WaveformViewerOverlay = SNew(SWaveformViewerOverlay, InTransportCoordinator);
+	WaveformViewerOverlay = SNew(SWaveformViewerOverlay, InTransportCoordinator, WaveformTransformationsOverlay.ToSharedRef());
 	WaveformViewerOverlay->OnNewMouseDelta.BindSP(InZoomManager, &FWaveformEditorZoomController::ZoomByDelta);
 }
 
diff --git a/Engine/Plugins/Editor/WaveformEditor/Source/WaveformEditorWidgets/Private/SWaveformTransformationRenderLayer.cpp b/Engine/Plugins/Editor/WaveformEditor/Source/WaveformEditorWidgets/Private/SWaveformTransformationRenderLayer.cpp
new file mode 100644
--- /dev/null
+++ b/Engine/Plugins/Editor/WaveformEditor/Source/WaveformEditorWidgets/Private/SWaveformTransformationRenderLayer.cpp
@@ -0,0 +1,70 @@
+// Copyright Epic Games, Inc. All Rights Reserved.
+
+#include "SWaveformTransformationRenderLayer.h"
+
+#include "Styling/AppStyle.h"
+
+void SWaveformTransformationRenderLayer::PrivateRegisterAttributes(FSlateAttributeInitializer& AttributeInitializer)
+{
+}
+
+void SWaveformTransformationRenderLayer::Construct(const FArguments& InArgs, TSharedRef<IWaveformTransformationRenderer> InTransformationRenderer)
+{
+	TransformationRenderer = InTransformationRenderer;
+}
+
+int32 SWaveformTransformationRenderLayer::OnPaint(const FPaintArgs& Args, const FGeometry& AllottedGeometry, const FSlateRect& MyCullingRect, FSlateWindowElementList& OutDrawElements, int32 LayerId, const FWidgetStyle& InWidgetStyle, bool bParentEnabled) const
+{
+	return TransformationRenderer->OnPaint(Args, AllottedGeometry, MyCullingRect, OutDrawElements, LayerId, InWidgetStyle, bParentEnabled);
+}
+
+FVector2D SWaveformTransformationRenderLayer::ComputeDesiredSize(float) const
+{
+	return FVector2D(1280, 720);
+}
+
+FCursorReply SWaveformTransformationRenderLayer::OnCursorQuery(const FGeometry& MyGeometry, const FPointerEvent& CursorEvent) const
+{
+	return TransformationRenderer->OnCursorQuery(MyGeometry, CursorEvent);
+}
+
+void SWaveformTransformationRenderLayer::Tick(const FGeometry& AllottedGeometry, const double InCurrentTime, const float InDeltaTime)
+{
+	return TransformationRenderer->Tick(AllottedGeometry, InCurrentTime, InDeltaTime);
+}
+
+FReply SWaveformTransformationRenderLayer::OnMouseButtonDown(const FGeometry& MyGeometry, const FPointerEvent& MouseEvent)
+{
+	return TransformationRenderer->OnMouseButtonDown(*this, MyGeometry, MouseEvent);
+}
+
+FReply SWaveformTransformationRenderLayer::OnMouseButtonUp(const FGeometry& MyGeometry, const FPointerEvent& MouseEvent)
+{
+	return TransformationRenderer->OnMouseButtonUp(*this, MyGeometry, MouseEvent);
+}
+
+FReply SWaveformTransformationRenderLayer::OnMouseMove(const FGeometry& MyGeometry, const FPointerEvent& MouseEvent)
+{
+	return TransformationRenderer->OnMouseMove(*this, MyGeometry, MouseEvent);
+}
+
+FReply SWaveformTransformationRenderLayer::OnMouseWheel(const FGeometry& MyGeometry, const FPointerEvent& MouseEvent)
+{
+	return TransformationRenderer->OnMouseWheel(*this, MyGeometry, MouseEvent);
+}
+
+FReply SWaveformTransformationRenderLayer::OnMouseButtonDoubleClick(const FGeometry& InMyGeometry, const FPointerEvent& InMouseEvent)
+{
+	return TransformationRenderer->OnMouseButtonDoubleClick(*this, InMyGeometry, InMouseEvent);
+}
+
+void SWaveformTransformationRenderLayer::OnMouseEnter(const FGeometry& MyGeometry, const FPointerEvent& MouseEvent)
+{
+	return TransformationRenderer->OnMouseEnter(MyGeometry, MouseEvent);
+}
+
+void SWaveformTransformationRenderLayer::OnMouseLeave(const FPointerEvent& MouseEvent)
+{
+	return TransformationRenderer->OnMouseLeave(MouseEvent);
+}
+
diff --git a/Engine/Plugins/Editor/WaveformEditor/Source/WaveformEditorWidgets/Private/SWaveformTransformationsOverlay.cpp b/Engine/Plugins/Editor/WaveformEditor/Source/WaveformEditorWidgets/Private/SWaveformTransformationsOverlay.cpp
--- a/Engine/Plugins/Editor/WaveformEditor/Source/WaveformEditorWidgets/Private/SWaveformTransformationsOverlay.cpp
+++ b/Engine/Plugins/Editor/WaveformEditor/Source/WaveformEditorWidgets/Private/SWaveformTransformationsOverlay.cpp
@@ -3,34 +3,165 @@
 #include "SWaveformTransformationsOverlay.h"
 
 #include "Widgets/SOverlay.h"
+#include "WaveformEditorTransportCoordinator.h"
+#include "SWaveformTransformationRenderLayer.h"
 
-void SWaveformTransformationsOverlay::Construct(const FArguments& InArgs, TArrayView<TSharedPtr<SWidget>> InTransformationLayers)
+void SWaveformTransformationsOverlay::Construct(const FArguments& InArgs, TArrayView<const FTransformationLayerRenderInfo> InTransformationRenderers, TSharedRef<FWaveformEditorTransportCoordinator> InTransportCoordinator)
 {
-	TransformationLayers = InTransformationLayers;
+	TransformationRenderers = InTransformationRenderers;
+	TransportCoordinator = InTransportCoordinator;
 	CreateLayout();
 }
 
 void SWaveformTransformationsOverlay::CreateLayout()
 {
+	const int32 NumLayers = TransformationRenderers.Num();
+	TransformationLayers.Empty();
+	TransformationLayers.SetNumZeroed(NumLayers);
+	LayersSlots.Empty();
+	LayersSlots.SetNumUninitialized(NumLayers);
+
 	ChildSlot
 	[
 		SAssignNew(MainOverlayPtr, SOverlay)
 	];
 
-	for (TSharedPtr<SWidget> Layer : TransformationLayers)
+	for (int32 i = 0; i < NumLayers; ++i)
 	{
-		if (Layer)
+		const FTransformationLayerRenderInfo& LayerRenderInfo = TransformationRenderers[i];
+
+		if (LayerRenderInfo.Key)
 		{
+			float LeftAnchor = TransportCoordinator->ConvertAbsoluteRatioToZoomed(LayerRenderInfo.Value.Key);
+			float RightAnchor = TransportCoordinator->ConvertAbsoluteRatioToZoomed(LayerRenderInfo.Value.Value);
+
+			SConstraintCanvas::FSlot* SlotPtr = nullptr;
+
+			TSharedPtr<SConstraintCanvas> ConstraintCanvasPtr = SNew(SConstraintCanvas)
+				+ SConstraintCanvas::Slot()
+				.Anchors(FAnchors(LeftAnchor, 0.f, RightAnchor, 1.f))
+				.Expose(SlotPtr)
+				[
+					SAssignNew(TransformationLayers[i], SWaveformTransformationRenderLayer, LayerRenderInfo.Key.ToSharedRef())
+				];
+
 			MainOverlayPtr->AddSlot()
 			[
-				Layer.ToSharedRef()
+				ConstraintCanvasPtr.ToSharedRef()
 			];
+
+			LayersSlots[i] = SlotPtr;
+
+		}
+		else
+		{
+			TransformationLayers[i] = nullptr;
+			LayersSlots[i] = nullptr;
+		}
+	}
+}
+
+
+void SWaveformTransformationsOverlay::UpdateAnchors()
+{
+	check(TransformationRenderers.Num() == LayersSlots.Num())
+	
+	for (int32 i = 0; i < TransformationRenderers.Num(); ++i)
+	{
+		const FTransformationLayerRenderInfo& Layer = TransformationRenderers[i];
+
+		if (Layer.Key)
+		{
+			float LeftAnchor = TransportCoordinator->ConvertAbsoluteRatioToZoomed(Layer.Value.Key);
+			float RightAnchor = TransportCoordinator->ConvertAbsoluteRatioToZoomed(Layer.Value.Value);
+
+			LayersSlots[i]->SetAnchors(FAnchors(LeftAnchor, 0.f, RightAnchor, 1.f));
 		}
 	}
 }
 
-void SWaveformTransformationsOverlay::OnLayerChainUpdate(TSharedPtr<SWidget>* FirstLayerPtr, const int32 NLayers)
+FReply SWaveformTransformationsOverlay::RouteMouseInput(WidgetMouseInputFunction InputFunction, const FGeometry& MyGeometry, const FPointerEvent& MouseEvent)
 {
-	TransformationLayers = MakeArrayView(FirstLayerPtr, NLayers);
+	FReply TransformationInteraction = FReply::Unhandled();
+
+	//input is propagated from topmost transformation to lowermost
+	for (int32 i = TransformationLayers.Num() - 1; i >= 0; i--)
+	{
+		const TSharedPtr<SWaveformTransformationRenderLayer> LayerWidget = TransformationLayers[i];
+
+		if (LayerWidget)
+		{
+			TransformationInteraction = (LayerWidget.Get()->*InputFunction)(LayerWidget->GetTickSpaceGeometry(), MouseEvent);
+
+			if (TransformationInteraction.IsEventHandled())
+			{
+				return TransformationInteraction;
+			}
+		}
+	}
+		
+	return TransformationInteraction;
+}
+
+void SWaveformTransformationsOverlay::OnLayerChainGenerated(FTransformationLayerRenderInfo* FirstLayerPtr, const int32 NLayers)
+{
+	TransformationRenderers = MakeArrayView(FirstLayerPtr, NLayers);
 	CreateLayout();
 }
+
+void SWaveformTransformationsOverlay::UpdateLayerConstraints()
+{
+	UpdateAnchors();
+}
+
+void SWaveformTransformationsOverlay::OnNewWaveformDisplayRange(const TRange<float> NewDisplayRange)
+{
+	UpdateAnchors();
+}
+
+FReply SWaveformTransformationsOverlay::OnMouseButtonDown(const FGeometry& MyGeometry, const FPointerEvent& MouseEvent)
+{
+	return RouteMouseInput(&SWidget::OnMouseButtonDown, MyGeometry, MouseEvent);
+}
+
+FReply SWaveformTransformationsOverlay::OnMouseButtonUp(const FGeometry& MyGeometry, const FPointerEvent& MouseEvent)
+{
+	return RouteMouseInput(&SWidget::OnMouseButtonUp, MyGeometry, MouseEvent);
+}
+
+FReply SWaveformTransformationsOverlay::OnMouseMove(const FGeometry& MyGeometry, const FPointerEvent& MouseEvent)
+{
+	return RouteMouseInput(&SWidget::OnMouseMove, MyGeometry, MouseEvent);
+}
+
+FReply SWaveformTransformationsOverlay::OnMouseWheel(const FGeometry& MyGeometry, const FPointerEvent& MouseEvent)
+{
+	return RouteMouseInput(&SWidget::OnMouseWheel, MyGeometry, MouseEvent);
+}
+
+FCursorReply SWaveformTransformationsOverlay::OnCursorQuery(const FGeometry& MyGeometry, const FPointerEvent& CursorEvent) const
+{
+	return RouteCursorQuery(MyGeometry, CursorEvent);
+}
+
+FCursorReply SWaveformTransformationsOverlay::RouteCursorQuery(const FGeometry& MyGeometry, const FPointerEvent& CursorEvent) const
+{
+	FCursorReply CursorReply = FCursorReply::Unhandled();
+
+	for (int32 i = TransformationLayers.Num() - 1; i >= 0; i--)
+	{
+		const TSharedPtr<SWaveformTransformationRenderLayer> LayerWidget = TransformationLayers[i];
+
+		if (LayerWidget)
+		{
+			CursorReply = LayerWidget->OnCursorQuery(LayerWidget->GetTickSpaceGeometry(), CursorEvent);
+
+			if (CursorReply.IsEventHandled())
+			{
+				return CursorReply;
+			}
+		}
+	}
+
+	return CursorReply;
+}
\ No newline at end of file
diff --git a/Engine/Plugins/Editor/WaveformEditor/Source/WaveformEditorWidgets/Private/SWaveformViewerOverlay.cpp b/Engine/Plugins/Editor/WaveformEditor/Source/WaveformEditorWidgets/Private/SWaveformViewerOverlay.cpp
--- a/Engine/Plugins/Editor/WaveformEditor/Source/WaveformEditorWidgets/Private/SWaveformViewerOverlay.cpp
+++ b/Engine/Plugins/Editor/WaveformEditor/Source/WaveformEditorWidgets/Private/SWaveformViewerOverlay.cpp
@@ -5,9 +5,10 @@
 #include "WaveformEditorTransportCoordinator.h"
 #include "Widgets/SLeafWidget.h"
 
-void SWaveformViewerOverlay::Construct(const FArguments& InArgs, TSharedRef<FWaveformEditorTransportCoordinator> InTransportCoordinator)
+void SWaveformViewerOverlay::Construct(const FArguments& InArgs, TSharedRef<FWaveformEditorTransportCoordinator> InTransportCoordinator, TSharedRef<SWaveformTransformationsOverlay> InTransformationsOverlay)
 {
 	TransportCoordinator = InTransportCoordinator;
+	TransformationsOverlay = InTransformationsOverlay;
 
 	Style = InArgs._Style;
 
@@ -20,13 +21,38 @@ void SWaveformViewerOverlay::Construct(const FArguments& InArgs, TSharedRef<FWav
 
 FReply SWaveformViewerOverlay::OnMouseButtonDown(const FGeometry& MyGeometry, const FPointerEvent& MouseEvent)
 {
-	return TransportCoordinator->ReceiveMouseButtonDown(*this, MyGeometry, MouseEvent);
+	FReply InteractionReply = TransformationsOverlay->OnMouseButtonDown(MyGeometry, MouseEvent);
+
+	if (!InteractionReply.IsEventHandled())
+	{
+		InteractionReply = TransportCoordinator->ReceiveMouseButtonDown(*this, MyGeometry, MouseEvent);
+	}
+	
+	return InteractionReply;
 }
 
 FReply SWaveformViewerOverlay::OnMouseWheel(const FGeometry& MyGeometry, const FPointerEvent& MouseEvent)
 {
-	OnNewMouseDelta.ExecuteIfBound(MouseEvent.GetWheelDelta());
-	return FReply::Handled();
+	FReply InteractionReply = TransformationsOverlay->OnMouseWheel(MyGeometry, MouseEvent);
+
+	if (!InteractionReply.IsEventHandled())
+	{
+		OnNewMouseDelta.ExecuteIfBound(MouseEvent.GetWheelDelta());
+		InteractionReply = FReply::Handled();
+	}
+
+	return InteractionReply;
+
+}
+
+FReply SWaveformViewerOverlay::OnMouseMove(const FGeometry& MyGeometry, const FPointerEvent& MouseEvent)
+{
+	return TransformationsOverlay->OnMouseMove(MyGeometry, MouseEvent);
+}
+
+FCursorReply SWaveformViewerOverlay::OnCursorQuery(const FGeometry& MyGeometry, const FPointerEvent& CursorEvent) const
+{
+	return TransformationsOverlay->OnCursorQuery(MyGeometry, CursorEvent);
 }
 
 int32 SWaveformViewerOverlay::OnPaint(const FPaintArgs& Args, const FGeometry& AllottedGeometry, const FSlateRect& MyCullingRect, FSlateWindowElementList& OutDrawElements, int32 LayerId, const FWidgetStyle& InWidgetStyle, bool bParentEnabled) const
diff --git a/Engine/Plugins/Editor/WaveformEditor/Source/WaveformEditorWidgets/Private/SWaveformViewerOverlay.h b/Engine/Plugins/Editor/WaveformEditor/Source/WaveformEditorWidgets/Private/SWaveformViewerOverlay.h
--- a/Engine/Plugins/Editor/WaveformEditor/Source/WaveformEditorWidgets/Private/SWaveformViewerOverlay.h
+++ b/Engine/Plugins/Editor/WaveformEditor/Source/WaveformEditorWidgets/Private/SWaveformViewerOverlay.h
@@ -6,6 +6,7 @@
 #include "Styling/SlateWidgetStyleAsset.h"
 #include "WaveformEditorSlateTypes.h"
 #include "WaveformEditorTransportCoordinator.h"
+#include "SWaveformTransformationsOverlay.h"
 #include "Widgets/SLeafWidget.h"
 
 DECLARE_DELEGATE_OneParam(FOnNewMouseDelta, const float /* new delta */)
@@ -25,10 +26,12 @@ public:
 
 	SLATE_END_ARGS()
 
-	void Construct(const FArguments& InArgs, TSharedRef<FWaveformEditorTransportCoordinator> InTransportCoordinator);
+	void Construct(const FArguments& InArgs, TSharedRef<FWaveformEditorTransportCoordinator> InTransportCoordinator, TSharedRef<SWaveformTransformationsOverlay> InTransformationsOverlay);
 
 	virtual FReply OnMouseButtonDown(const FGeometry& MyGeometry, const FPointerEvent& MouseEvent) override;
 	virtual FReply OnMouseWheel(const FGeometry& MyGeometry, const FPointerEvent& MouseEvent) override;
+	FReply OnMouseMove(const FGeometry& MyGeometry, const FPointerEvent& MouseEvent) override;
+	FCursorReply OnCursorQuery(const FGeometry& MyGeometry, const FPointerEvent& CursorEvent) const override;
 
 	FOnNewMouseDelta OnNewMouseDelta;
 
@@ -39,6 +42,7 @@ private:
 	int32 DrawPlayhead(const FGeometry& AllottedGeometry, FSlateWindowElementList& OutDrawElements, int32 LayerId) const;
 
 	TSharedPtr<FWaveformEditorTransportCoordinator> TransportCoordinator = nullptr;
+	TSharedPtr<SWaveformTransformationsOverlay> TransformationsOverlay = nullptr;
 
 	const FWaveformViewerOverlayStyle* Style = nullptr;
 
diff --git a/Engine/Plugins/Editor/WaveformEditor/Source/WaveformEditorWidgets/Private/WaveformEditorTransportCoordinator.cpp b/Engine/Plugins/Editor/WaveformEditor/Source/WaveformEditorWidgets/Private/WaveformEditorTransportCoordinator.cpp
--- a/Engine/Plugins/Editor/WaveformEditor/Source/WaveformEditorWidgets/Private/WaveformEditorTransportCoordinator.cpp
+++ b/Engine/Plugins/Editor/WaveformEditor/Source/WaveformEditorWidgets/Private/WaveformEditorTransportCoordinator.cpp
@@ -214,6 +214,16 @@ void FWaveformEditorTransportCoordinator::OnZoomLevelChanged(const uint8 NewLeve
 	UpdateZoomRatioAndDisplayRange(NewZoomRatio);
 }
 
+float FWaveformEditorTransportCoordinator::ConvertAbsoluteRatioToZoomed(const float InAbsoluteRatio) const
+{
+	return (InAbsoluteRatio - DisplayRange.GetLowerBoundValue()) / ZoomRatio;
+}
+
+float FWaveformEditorTransportCoordinator::ConvertZoomedRatioToAbsolute(const float InZoomedRatio) const
+{
+	return InZoomedRatio * ZoomRatio + DisplayRange.GetLowerBoundValue();
+}
+
 TRange<float> FWaveformEditorTransportCoordinator::GetDisplayRange() const
 {
 	return DisplayRange;
diff --git a/Engine/Plugins/Editor/WaveformEditor/Source/WaveformEditorWidgets/Public/IWaveformTransformationRenderer.h b/Engine/Plugins/Editor/WaveformEditor/Source/WaveformEditorWidgets/Public/IWaveformTransformationRenderer.h
new file mode 100644
--- /dev/null
+++ b/Engine/Plugins/Editor/WaveformEditor/Source/WaveformEditorWidgets/Public/IWaveformTransformationRenderer.h
@@ -0,0 +1,89 @@
+// Copyright Epic Games, Inc. All Rights Reserved.
+
+#pragma once
+
+#include "HAL/Platform.h"
+#include "IWaveformTransformation.h"
+#include "Sound/SoundWave.h"
+#include "Templates/Function.h"
+#include "UObject/UnrealType.h"
+
+class FCursorReply;
+class FPaintArgs;
+class FReply;
+class FSlateRect;
+class FSlateWindowElementList;
+class FWidgetStyle;
+class SWidget;
+struct FGeometry;
+struct FPointerEvent;
+
+struct WAVEFORMEDITORWIDGETS_API FWaveformTransformationRenderInfo
+{
+	float SampleRate = 0.f;
+	int32 NumChannels = 0;
+	uint32 StartFrameOffset = 0;	//the first frame this transformation can act upon
+	uint32 NumAvilableSamples = 0;	//the number of samples available to the transformation
+};
+
+class WAVEFORMEDITORWIDGETS_API IWaveformTransformationRenderer
+{
+public:
+	virtual ~IWaveformTransformationRenderer() {};
+
+	/* SWidget interceptor methods */
+	virtual int32 OnPaint(const FPaintArgs& Args, const FGeometry& AllottedGeometry, const FSlateRect& MyCullingRect, FSlateWindowElementList& OutDrawElements, int32 LayerId, const FWidgetStyle& InWidgetStyle, bool bParentEnabled) const = 0;
+	virtual void Tick(const FGeometry& AllottedGeometry, const double InCurrentTime, const float InDeltaTime) = 0;
+	virtual FReply OnMouseButtonDown(SWidget& OwnerWidget, const FGeometry& MyGeometry, const FPointerEvent& MouseEvent) = 0;
+	virtual FReply OnMouseButtonUp(SWidget& OwnerWidget, const FGeometry& MyGeometry, const FPointerEvent& MouseEvent) = 0;
+	virtual FReply OnMouseButtonDoubleClick(SWidget& OwnerWidget, const FGeometry& InMyGeometry, const FPointerEvent& InMouseEvent) = 0;
+	virtual FReply OnMouseMove(SWidget& OwnerWidget, const FGeometry& MyGeometry, const FPointerEvent& MouseEvent) = 0;
+	virtual FReply OnMouseWheel(SWidget& OwnerWidget, const FGeometry& MyGeometry, const FPointerEvent& MouseEvent) = 0;
+	virtual void OnMouseEnter(const FGeometry& MyGeometry, const FPointerEvent& MouseEvent) = 0;
+	virtual void OnMouseLeave(const FPointerEvent& MouseEvent) = 0;
+	virtual FCursorReply OnCursorQuery(const FGeometry& MyGeometry, const FPointerEvent& CursorEvent) const = 0;
+
+	/* Gets info about the transformation being rendered */
+	void SetTransformationWaveInfo(FWaveformTransformationRenderInfo InWaveInfo)
+	{
+		TransformationWaveInfo = InWaveInfo;
+	}
+
+	/* Sets a callback to notify when a property has been changed by the renderer */
+	void SetTransformationNotifier(const TFunction<void(FPropertyChangedEvent&, FEditPropertyChain*)>& InTransformationChangeNotifier)
+	{
+		TransformationChangeNotifier = InTransformationChangeNotifier;
+	}
+
+protected:
+
+	/* Notifies that a property has been changed through TransformationChangeNotifier  */
+	void NotifyTransformationPropertyChanged(const TObjectPtr<UWaveformTransformationBase> EditedTransformation, const FName& PropertyName, const EPropertyChangeType::Type InChangeType)
+	{
+		FProperty* Property = EditedTransformation->GetClass()->FindPropertyByName(PropertyName);
+		check(Property);
+
+		const USoundWave* ParentSoundWave = EditedTransformation->GetTypedOuter<USoundWave>();
+		check(ParentSoundWave)
+
+		FProperty* TransformationsProperty = ParentSoundWave->GetClass()->FindPropertyByName(GET_MEMBER_NAME_CHECKED(USoundWave, Transformations));
+		check(TransformationsProperty)
+
+		FEditPropertyChain EditChain;
+		EditChain.AddHead(TransformationsProperty);
+		EditChain.AddTail(Property);
+		EditChain.SetActivePropertyNode(Property);
+
+		FPropertyChangedEvent PropertyEvent(Property, InChangeType);
+
+		EditedTransformation->PostEditChangeProperty(PropertyEvent);
+
+		TransformationChangeNotifier(PropertyEvent, &EditChain);
+	}
+
+	FWaveformTransformationRenderInfo TransformationWaveInfo;
+
+private:
+	TFunction<void(FPropertyChangedEvent&, FEditPropertyChain*)> TransformationChangeNotifier;
+
+};
\ No newline at end of file
diff --git a/Engine/Plugins/Editor/WaveformEditor/Source/WaveformEditorWidgets/Public/SWaveformTransformationRenderLayer.h b/Engine/Plugins/Editor/WaveformEditor/Source/WaveformEditorWidgets/Public/SWaveformTransformationRenderLayer.h
new file mode 100644
--- /dev/null
+++ b/Engine/Plugins/Editor/WaveformEditor/Source/WaveformEditorWidgets/Public/SWaveformTransformationRenderLayer.h
@@ -0,0 +1,33 @@
+// Copyright Epic Games, Inc. All Rights Reserved.
+
+#pragma once
+
+#include "Widgets/SCompoundWidget.h"
+#include "IWaveformTransformationRenderer.h"
+
+class WAVEFORMEDITORWIDGETS_API SWaveformTransformationRenderLayer : public SCompoundWidget
+{
+	SLATE_DECLARE_WIDGET(SWaveformTransformationRenderLayer, SCompoundWidget)
+
+public:
+	SLATE_BEGIN_ARGS(SWaveformTransformationRenderLayer) {}
+		SLATE_DEFAULT_SLOT(FArguments, InArgs)
+	SLATE_END_ARGS()
+
+	void Construct(const FArguments& InArgs, TSharedRef<IWaveformTransformationRenderer> InTransformationRenderer);
+	
+	virtual int32 OnPaint(const FPaintArgs& Args, const FGeometry& AllottedGeometry, const FSlateRect& MyCullingRect, FSlateWindowElementList& OutDrawElements, int32 LayerId, const FWidgetStyle& InWidgetStyle, bool bParentEnabled) const override;
+	void Tick(const FGeometry& AllottedGeometry, const double InCurrentTime, const float InDeltaTime) override;
+	FVector2D ComputeDesiredSize(float) const override;
+	FCursorReply OnCursorQuery(const FGeometry& MyGeometry, const FPointerEvent& CursorEvent) const override;
+	FReply OnMouseButtonDown(const FGeometry& MyGeometry, const FPointerEvent& MouseEvent) override;
+	FReply OnMouseButtonUp(const FGeometry& MyGeometry, const FPointerEvent& MouseEvent) override;
+	FReply OnMouseMove(const FGeometry& MyGeometry, const FPointerEvent& MouseEvent) override;
+	FReply OnMouseWheel(const FGeometry& MyGeometry, const FPointerEvent& MouseEvent) override;
+	FReply OnMouseButtonDoubleClick(const FGeometry& InMyGeometry, const FPointerEvent& InMouseEvent) override;
+	void OnMouseEnter(const FGeometry& MyGeometry, const FPointerEvent& MouseEvent) override;
+	void OnMouseLeave(const FPointerEvent& MouseEvent) override;
+
+private:
+	TSharedPtr<IWaveformTransformationRenderer> TransformationRenderer;
+};
\ No newline at end of file
diff --git a/Engine/Plugins/Editor/WaveformEditor/Source/WaveformEditorWidgets/Public/SWaveformTransformationsOverlay.h b/Engine/Plugins/Editor/WaveformEditor/Source/WaveformEditorWidgets/Public/SWaveformTransformationsOverlay.h
--- a/Engine/Plugins/Editor/WaveformEditor/Source/WaveformEditorWidgets/Public/SWaveformTransformationsOverlay.h
+++ b/Engine/Plugins/Editor/WaveformEditor/Source/WaveformEditorWidgets/Public/SWaveformTransformationsOverlay.h
@@ -3,8 +3,14 @@
 #pragma once
 
 #include "Widgets/SCompoundWidget.h"
+#include "Widgets/Layout/SConstraintCanvas.h"
 
+class FWaveformEditorTransportCoordinator;
+class IWaveformTransformationRenderer;
 class SOverlay;
+class SWaveformTransformationRenderLayer;
+
+using FTransformationLayerRenderInfo = TPair<TSharedPtr<IWaveformTransformationRenderer>, TPair<float,float>>;
 
 class WAVEFORMEDITORWIDGETS_API SWaveformTransformationsOverlay : public SCompoundWidget
 {
@@ -14,12 +20,29 @@ public:
 		SLATE_DEFAULT_SLOT(FArguments, InArgs)
 	SLATE_END_ARGS()
 
-	void Construct(const FArguments& InArgs, TArrayView<TSharedPtr<SWidget>> InTransformationLayers);
-	void OnLayerChainUpdate(TSharedPtr<SWidget>* FirstLayerPtr, const int32 NLayers);
+	void Construct(const FArguments& InArgs, TArrayView< const FTransformationLayerRenderInfo> InTransformationRenderers, TSharedRef<FWaveformEditorTransportCoordinator> InTransportCoordinator);
+	void OnLayerChainGenerated(FTransformationLayerRenderInfo* FirstLayerPtr, const int32 NLayers);
+	void UpdateLayerConstraints();
+	void OnNewWaveformDisplayRange(const TRange<float> NewDisplayRange);
 
+	FReply OnMouseButtonDown(const FGeometry& MyGeometry, const FPointerEvent& MouseEvent) override;
+	FReply OnMouseButtonUp(const FGeometry& MyGeometry, const FPointerEvent& MouseEvent) override;
+	FReply OnMouseMove(const FGeometry& MyGeometry, const FPointerEvent& MouseEvent) override;
+	FReply OnMouseWheel(const FGeometry& MyGeometry, const FPointerEvent& MouseEvent) override;
+	FCursorReply OnCursorQuery(const FGeometry& MyGeometry, const FPointerEvent& CursorEvent) const override;
+	
 private:
 	void CreateLayout();
+	void UpdateAnchors();
 	
 	TSharedPtr<SOverlay> MainOverlayPtr;
-	TArrayView<TSharedPtr<SWidget>> TransformationLayers;
+	TArray<TSharedPtr<SWaveformTransformationRenderLayer>> TransformationLayers;
+	TArray<SConstraintCanvas::FSlot*> LayersSlots;
+	TArrayView<const FTransformationLayerRenderInfo> TransformationRenderers;
+	TSharedPtr<FWaveformEditorTransportCoordinator> TransportCoordinator = nullptr;
+
+	typedef FReply(SWidget::*WidgetMouseInputFunction)(const FGeometry&, const FPointerEvent&);
+	
+	FReply RouteMouseInput(WidgetMouseInputFunction InputFunction, const FGeometry& MyGeometry, const FPointerEvent& MouseEvent);
+	FCursorReply RouteCursorQuery(const FGeometry& MyGeometry, const FPointerEvent& CursorEvent) const; 
 };
\ No newline at end of file
diff --git a/Engine/Plugins/Editor/WaveformEditor/Source/WaveformEditorWidgets/Public/WaveformEditorTransportCoordinator.h b/Engine/Plugins/Editor/WaveformEditor/Source/WaveformEditorWidgets/Public/WaveformEditorTransportCoordinator.h
--- a/Engine/Plugins/Editor/WaveformEditor/Source/WaveformEditorWidgets/Public/WaveformEditorTransportCoordinator.h
+++ b/Engine/Plugins/Editor/WaveformEditor/Source/WaveformEditorWidgets/Public/WaveformEditorTransportCoordinator.h
@@ -42,6 +42,9 @@ public:
 	void ReceivePlayBackRatio(const float NewRatio);
 	void OnZoomLevelChanged(const uint8 NewLevel);
 
+	float ConvertAbsoluteRatioToZoomed(const float InAbsoluteRatio) const;
+	float ConvertZoomedRatioToAbsolute(const float InZoomedRatio) const;
+
 	TRange<float> GetDisplayRange() const;
 
 	void Stop();
diff --git a/Engine/Plugins/Editor/WaveformEditor/Source/WaveformTransformations/Public/WaveformTransformationTrimFade.h b/Engine/Plugins/Editor/WaveformEditor/Source/WaveformTransformations/Public/WaveformTransformationTrimFade.h
--- a/Engine/Plugins/Editor/WaveformEditor/Source/WaveformTransformations/Public/WaveformTransformationTrimFade.h
+++ b/Engine/Plugins/Editor/WaveformEditor/Source/WaveformTransformations/Public/WaveformTransformationTrimFade.h
@@ -31,6 +31,8 @@ class WAVEFORMTRANSFORMATIONS_API UWaveformTransformationTrimFade : public UWave
 {
 	GENERATED_BODY()
 
+public:
+
 	UPROPERTY(EditAnywhere, Category = "Trim", meta=(ClampMin = 0.0))
 	double StartTime = 0.0;
 
diff --git a/Engine/Plugins/Editor/WaveformEditor/Source/WaveformTransformationsWidgets/Private/SWaveformTransformationDurationHighlight.h b/Engine/Plugins/Editor/WaveformEditor/Source/WaveformTransformationsWidgets/Private/SWaveformTransformationDurationHighlight.h
deleted file mode 100644
diff --git a/Engine/Plugins/Editor/WaveformEditor/Source/WaveformTransformationsWidgets/Private/SWaveformTransformationRenderLayer.cpp b/Engine/Plugins/Editor/WaveformEditor/Source/WaveformTransformationsWidgets/Private/SWaveformTransformationRenderLayer.cpp
deleted file mode 100644
diff --git a/Engine/Plugins/Editor/WaveformEditor/Source/WaveformTransformationsWidgets/Private/SWaveformTransformationTrimFadeLayer.h b/Engine/Plugins/Editor/WaveformEditor/Source/WaveformTransformationsWidgets/Private/SWaveformTransformationTrimFadeLayer.h
deleted file mode 100644
diff --git a/Engine/Plugins/Editor/WaveformEditor/Source/WaveformTransformationsWidgets/Private/WaveformTransformationDurationRenderer.h b/Engine/Plugins/Editor/WaveformEditor/Source/WaveformTransformationsWidgets/Private/WaveformTransformationDurationRenderer.h
new file mode 100644
--- /dev/null
+++ b/Engine/Plugins/Editor/WaveformEditor/Source/WaveformTransformationsWidgets/Private/WaveformTransformationDurationRenderer.h
@@ -0,0 +1,63 @@
+// Copyright Epic Games, Inc. All Rights Reserved.
+
+#pragma once
+
+#include "WaveformTransformationRendererBase.h"
+#include "WaveformEditorRenderData.h"
+
+namespace WaveformTransformationDurationHiglightParams
+{
+	const FLazyName BackgroundBrushName = TEXT("WhiteBrush");
+	const FLinearColor BoxColor = FLinearColor(0.f, 0.f, 0.f, 0.7f);
+}
+
+class FWaveformTransformationDurationRenderer : public FWaveformTransformationRendererBase
+{
+public:
+	explicit FWaveformTransformationDurationRenderer(const TSharedRef<FWaveformEditorRenderData> InRenderData)
+	{
+		WaveformRenderData = InRenderData;
+	}
+
+	virtual int32 OnPaint(const FPaintArgs& Args, const FGeometry& AllottedGeometry, const FSlateRect& MyCullingRect, FSlateWindowElementList& OutDrawElements, int32 LayerId, const FWidgetStyle& InWidgetStyle, bool bParentEnabled) const override
+	{
+		const float StartTimeRatio = (float)TransformationWaveInfo.StartFrameOffset / WaveformRenderData->GetOriginalWaveformFrames() / TransformationWaveInfo.NumChannels;
+		const uint32 EndSample = TransformationWaveInfo.StartFrameOffset + TransformationWaveInfo.NumAvilableSamples;
+		const float EndTimeRatio = (float)EndSample / WaveformRenderData->GetOriginalWaveformFrames() / TransformationWaveInfo.NumChannels;
+
+		const bool bRenderLeftBox = StartTimeRatio >= 0.f;
+		const bool bRenderRightBox = EndTimeRatio <= 1.f;
+
+		if (bRenderLeftBox)
+		{
+			const float RightMarginX = StartTimeRatio * AllottedGeometry.Size.X;
+
+			FSlateDrawElement::MakeBox(
+				OutDrawElements,
+				++LayerId,
+				AllottedGeometry.ToPaintGeometry(FVector2D(0.f, 0.f), FVector2D(RightMarginX, AllottedGeometry.Size.Y)),
+				FAppStyle::GetBrush(WaveformTransformationDurationHiglightParams::BackgroundBrushName),
+				ESlateDrawEffect::None,
+				WaveformTransformationDurationHiglightParams::BoxColor);
+
+		}
+
+		if (bRenderRightBox)
+		{
+			const float LeftMarginX = EndTimeRatio * AllottedGeometry.Size.X;
+
+			FSlateDrawElement::MakeBox(
+				OutDrawElements,
+				++LayerId,
+				AllottedGeometry.ToPaintGeometry(FVector2D(LeftMarginX, 0.f), FVector2D(AllottedGeometry.Size.X, AllottedGeometry.Size.Y)),
+				FAppStyle::GetBrush(WaveformTransformationDurationHiglightParams::BackgroundBrushName),
+				ESlateDrawEffect::None,
+				WaveformTransformationDurationHiglightParams::BoxColor);
+
+		}
+
+		return LayerId;
+	}
+private:
+	TSharedPtr<FWaveformEditorRenderData> WaveformRenderData = nullptr;
+};
\ No newline at end of file
diff --git a/Engine/Plugins/Editor/WaveformEditor/Source/WaveformTransformationsWidgets/Private/WaveformTransformationRenderLayerFactory.cpp b/Engine/Plugins/Editor/WaveformEditor/Source/WaveformTransformationsWidgets/Private/WaveformTransformationRenderLayerFactory.cpp
--- a/Engine/Plugins/Editor/WaveformEditor/Source/WaveformTransformationsWidgets/Private/WaveformTransformationRenderLayerFactory.cpp
+++ b/Engine/Plugins/Editor/WaveformEditor/Source/WaveformTransformationsWidgets/Private/WaveformTransformationRenderLayerFactory.cpp
@@ -4,46 +4,36 @@
 
 #include "IWaveformTransformation.h"
 #include "Styling/AppStyle.h"
-#include "SWaveformTransformationDurationHighlight.h"
-#include "SWaveformTransformationRenderLayer.h"
-#include "SWaveformTransformationTrimFadeLayer.h"
 #include "WaveformEditorTransportCoordinator.h"
-#include "WaveformEditorZoomController.h"
+#include "WaveformTransformationDurationRenderer.h"
 #include "WaveformTransformationTrimFade.h"
+#include "WaveformTransformationTrimFadeRenderer.h"
 
-FWaveformTransformationRenderLayerFactory::FWaveformTransformationRenderLayerFactory(TSharedRef<FWaveformEditorRenderData> InWaveformRenderData, TSharedRef<FWaveformEditorTransportCoordinator> InTransportCoordinator, TSharedRef<FWaveformEditorZoomController> InZoomController)
+FWaveformTransformationRenderLayerFactory::FWaveformTransformationRenderLayerFactory(TSharedRef<FWaveformEditorRenderData> InWaveformRenderData, 
+	TFunction<void(FPropertyChangedEvent&, FEditPropertyChain*)> InTransformationChangeNotifier)
 	: WaveformRenderData(InWaveformRenderData)
-	, TransportCoordinator(InTransportCoordinator)
-	, ZoomController(InZoomController)
+	, TransformationChangeNotifier(InTransformationChangeNotifier)
 {
 }
 
 
-TSharedPtr<SWaveformTransformationRenderLayer> FWaveformTransformationRenderLayerFactory::Create(TObjectPtr<UWaveformTransformationBase> InTransformationToRender) const
+TSharedPtr<IWaveformTransformationRenderer> FWaveformTransformationRenderLayerFactory::Create(TObjectPtr<UWaveformTransformationBase> InTransformationToRender) const
 {
-	TSharedPtr<SWaveformTransformationRenderLayer> OutWidget = nullptr;
+	TSharedPtr<IWaveformTransformationRenderer> OutRenderer = nullptr;
 	UClass* TransformationClass = InTransformationToRender->GetClass();
 
 	if (TransformationClass == UWaveformTransformationTrimFade::StaticClass())
 	{
-		TSharedPtr<SWaveformTransformationTrimFadeLayer> TrimFadeLayer = SNew(SWaveformTransformationTrimFadeLayer, Cast<UWaveformTransformationTrimFade>(InTransformationToRender), WaveformRenderData.ToSharedRef());
-		TrimFadeLayer->OnZoomLevelChanged(ZoomController->GetZoomRatio());
-		TrimFadeLayer->UpdateDisplayRange(TransportCoordinator->GetDisplayRange());
-		TransportCoordinator->OnDisplayRangeUpdated.AddSP(TrimFadeLayer.Get(), &SWaveformTransformationTrimFadeLayer::UpdateDisplayRange);
-		ZoomController->OnZoomRatioChanged.AddSP(TrimFadeLayer.Get(), &SWaveformTransformationTrimFadeLayer::OnZoomLevelChanged);
-		OutWidget = TrimFadeLayer;
+		TSharedPtr<IWaveformTransformationRenderer> TrimFadeLayer = MakeShared<FWaveformTransformationTrimFadeRenderer>(Cast<UWaveformTransformationTrimFade>(InTransformationToRender));
+		TrimFadeLayer->SetTransformationNotifier(TransformationChangeNotifier);
+		OutRenderer = TrimFadeLayer;
 	}
 
-	return OutWidget;
+	return OutRenderer;
 }
 
-TSharedPtr<SWaveformTransformationRenderLayer> FWaveformTransformationRenderLayerFactory::CreateDurationHiglightLayer() const
+TSharedPtr<IWaveformTransformationRenderer> FWaveformTransformationRenderLayerFactory::CreateDurationHiglightLayer() const
 {
-	TSharedPtr<SWaveformTransformationDurationHiglight> DurationHiglightLayer = SNew(SWaveformTransformationDurationHiglight, WaveformRenderData.ToSharedRef());
-	DurationHiglightLayer->OnZoomLevelChanged(ZoomController->GetZoomRatio());
-	DurationHiglightLayer->UpdateDisplayRange(TransportCoordinator->GetDisplayRange());
-	TransportCoordinator->OnDisplayRangeUpdated.AddSP(DurationHiglightLayer.Get(), &SWaveformTransformationDurationHiglight::UpdateDisplayRange);
-	ZoomController->OnZoomRatioChanged.AddSP(DurationHiglightLayer.Get(), &SWaveformTransformationDurationHiglight::OnZoomLevelChanged);
-
+	TSharedPtr<IWaveformTransformationRenderer> DurationHiglightLayer = MakeShared<FWaveformTransformationDurationRenderer>(WaveformRenderData.ToSharedRef());
 	return DurationHiglightLayer;
 }
diff --git a/Engine/Plugins/Editor/WaveformEditor/Source/WaveformTransformationsWidgets/Private/WaveformTransformationRendererBase.cpp b/Engine/Plugins/Editor/WaveformEditor/Source/WaveformTransformationsWidgets/Private/WaveformTransformationRendererBase.cpp
new file mode 100644
--- /dev/null
+++ b/Engine/Plugins/Editor/WaveformEditor/Source/WaveformTransformationsWidgets/Private/WaveformTransformationRendererBase.cpp
@@ -0,0 +1,38 @@
+// Copyright Epic Games, Inc. All Rights Reserved.
+
+#include "WaveformTransformationRendererBase.h"
+
+int32 FWaveformTransformationRendererBase::OnPaint(const FPaintArgs& Args, const FGeometry& AllottedGeometry, const FSlateRect& MyCullingRect, FSlateWindowElementList& OutDrawElements, int32 LayerId, const FWidgetStyle& InWidgetStyle, bool bParentEnabled) const
+{
+	return LayerId;
+}
+
+FReply FWaveformTransformationRendererBase::OnMouseButtonDown(SWidget& OwnerWidget, const FGeometry& MyGeometry, const FPointerEvent& MouseEvent)
+{
+	return FReply::Unhandled();
+}
+
+FReply FWaveformTransformationRendererBase::OnMouseButtonUp(SWidget& OwnerWidget, const FGeometry& MyGeometry, const FPointerEvent& MouseEvent)
+{
+	return FReply::Unhandled();
+}
+
+FReply FWaveformTransformationRendererBase::OnMouseButtonDoubleClick(SWidget& OwnerWidget, const FGeometry& InMyGeometry, const FPointerEvent& InMouseEvent)
+{
+	return FReply::Unhandled();
+}
+
+FReply FWaveformTransformationRendererBase::OnMouseMove(SWidget& OwnerWidget, const FGeometry& MyGeometry, const FPointerEvent& MouseEvent)
+{
+	return FReply::Unhandled();
+}
+
+FReply FWaveformTransformationRendererBase::OnMouseWheel(SWidget& OwnerWidget, const FGeometry& MyGeometry, const FPointerEvent& MouseEvent)
+{
+	return FReply::Unhandled();
+}
+
+FCursorReply FWaveformTransformationRendererBase::OnCursorQuery(const FGeometry& MyGeometry, const FPointerEvent& CursorEvent) const
+{
+	return FCursorReply::Unhandled();
+}
diff --git a/Engine/Plugins/Editor/WaveformEditor/Source/WaveformTransformationsWidgets/Private/WaveformTransformationRendererBase.h b/Engine/Plugins/Editor/WaveformEditor/Source/WaveformTransformationsWidgets/Private/WaveformTransformationRendererBase.h
new file mode 100644
--- /dev/null
+++ b/Engine/Plugins/Editor/WaveformEditor/Source/WaveformTransformationsWidgets/Private/WaveformTransformationRendererBase.h
@@ -0,0 +1,22 @@
+// Copyright Epic Games, Inc. All Rights Reserved.
+
+#pragma once
+
+#include "IWaveformTransformationRenderer.h"
+
+class FWaveformTransformationRendererBase : public IWaveformTransformationRenderer
+{
+public:
+	FWaveformTransformationRendererBase() = default;
+
+	virtual int32 OnPaint(const FPaintArgs& Args, const FGeometry& AllottedGeometry, const FSlateRect& MyCullingRect, FSlateWindowElementList& OutDrawElements, int32 LayerId, const FWidgetStyle& InWidgetStyle, bool bParentEnabled) const override;
+	virtual void Tick(const FGeometry& AllottedGeometry, const double InCurrentTime, const float InDeltaTime) override {};
+	virtual FReply OnMouseButtonDown(SWidget& OwnerWidget, const FGeometry& MyGeometry, const FPointerEvent& MouseEvent) override;
+	virtual FReply OnMouseButtonUp(SWidget& OwnerWidget, const FGeometry& MyGeometry, const FPointerEvent& MouseEvent) override;
+	virtual FReply OnMouseButtonDoubleClick(SWidget& OwnerWidget, const FGeometry& InMyGeometry, const FPointerEvent& InMouseEvent) override;
+	virtual FReply OnMouseMove(SWidget& OwnerWidget, const FGeometry& MyGeometry, const FPointerEvent& MouseEvent) override;
+	virtual FReply OnMouseWheel(SWidget& OwnerWidget, const FGeometry& MyGeometry, const FPointerEvent& MouseEvent) override;
+	virtual void OnMouseEnter(const FGeometry& MyGeometry, const FPointerEvent& MouseEvent) override {};
+	virtual void OnMouseLeave(const FPointerEvent& MouseEvent) override {};
+	virtual FCursorReply OnCursorQuery(const FGeometry& MyGeometry, const FPointerEvent& CursorEvent) const override;
+};
\ No newline at end of file
diff --git a/Engine/Plugins/Editor/WaveformEditor/Source/WaveformTransformationsWidgets/Private/WaveformTransformationTrimFadeRenderer.cpp b/Engine/Plugins/Editor/WaveformEditor/Source/WaveformTransformationsWidgets/Private/WaveformTransformationTrimFadeRenderer.cpp
new file mode 100644
--- /dev/null
+++ b/Engine/Plugins/Editor/WaveformEditor/Source/WaveformTransformationsWidgets/Private/WaveformTransformationTrimFadeRenderer.cpp
@@ -0,0 +1,149 @@
+// Copyright Epic Games, Inc. All Rights Reserved.
+
+#include "WaveformTransformationTrimFadeRenderer.h"
+
+FWaveformTransformationTrimFadeRenderer::FWaveformTransformationTrimFadeRenderer(const TObjectPtr<UWaveformTransformationTrimFade> TransformationToRender)
+{
+	check(TransformationToRender);
+	TrimFadeTransform = TransformationToRender;
+}
+
+int32 FWaveformTransformationTrimFadeRenderer::OnPaint(const FPaintArgs& Args, const FGeometry& AllottedGeometry, const FSlateRect& MyCullingRect, FSlateWindowElementList& OutDrawElements, int32 LayerId, const FWidgetStyle& InWidgetStyle, bool bParentEnabled) const
+{
+	LayerId = DrawTrimHandles(AllottedGeometry, OutDrawElements, LayerId);
+	LayerId = DrawFadeCurves(AllottedGeometry, OutDrawElements, LayerId);
+
+	return LayerId;
+}
+
+int32 FWaveformTransformationTrimFadeRenderer::DrawTrimHandles(const FGeometry& AllottedGeometry, FSlateWindowElementList& OutDrawElements, int32 LayerId) const
+{
+	const bool bRenderLowerBound = StartTimeHandleX >= 0.f;
+	const bool bRenderUpperBound = EndTimeHandleX <= AllottedGeometry.Size.X;
+
+	TArray<FVector2D> LinePoints;
+	LinePoints.SetNumUninitialized(2);
+
+	if (bRenderLowerBound)
+	{
+		LinePoints[0] = FVector2D(StartTimeHandleX, 0.f);
+		LinePoints[1] = FVector2D(StartTimeHandleX, AllottedGeometry.Size.Y);
+
+		FSlateDrawElement::MakeLines(
+			OutDrawElements,
+			++LayerId,
+			AllottedGeometry.ToPaintGeometry(),
+			LinePoints,
+			ESlateDrawEffect::None,
+			FLinearColor::Green,
+			false
+		);
+	}
+
+	if (bRenderUpperBound)
+	{
+		LinePoints[0] = FVector2D(EndTimeHandleX, 0.f);
+		LinePoints[1] = FVector2D(EndTimeHandleX, AllottedGeometry.Size.Y);
+
+		FSlateDrawElement::MakeLines(
+			OutDrawElements,
+			++LayerId,
+			AllottedGeometry.ToPaintGeometry(),
+			LinePoints,
+			ESlateDrawEffect::None,
+			FLinearColor::Red,
+			false
+		);
+	}
+
+	return LayerId;
+}
+
+int32 FWaveformTransformationTrimFadeRenderer::DrawFadeCurves(const FGeometry& AllottedGeometry, FSlateWindowElementList& OutDrawElements, int32 LayerId) const
+{
+	if (FadeInCurvePoints.Num() > 0)
+	{
+		FSlateDrawElement::MakeLines(
+			OutDrawElements,
+			++LayerId,
+			AllottedGeometry.ToPaintGeometry(),
+			FadeInCurvePoints,
+			ESlateDrawEffect::None,
+			FLinearColor::Yellow
+		);
+
+	}
+
+	if (FadeOutCurvePoints.Num() > 0)
+	{
+		FSlateDrawElement::MakeLines(
+			OutDrawElements,
+			++LayerId,
+			AllottedGeometry.ToPaintGeometry(),
+			FadeOutCurvePoints,
+			ESlateDrawEffect::None,
+			FLinearColor::Yellow
+		);
+
+	}
+
+	return LayerId;
+}
+
+void FWaveformTransformationTrimFadeRenderer::GenerateFadeCurves(const FGeometry& AllottedGeometry)
+{
+	const float FadeInFrames = TrimFadeTransform->StartFadeTime * TransformationWaveInfo.SampleRate;
+	const uint32 FadeInPixelLenght = FadeInFrames * PixelsPerFrame;
+	FadeInStartX = FMath::RoundToInt32(StartTimeHandleX);
+	FadeInEndX = FMath::RoundToInt32(FMath::Clamp(StartTimeHandleX + FadeInPixelLenght, StartTimeHandleX, EndTimeHandleX));
+	
+	const uint32 DisplayedFadeInPixelLenght = FadeInEndX - FadeInStartX;
+	FadeInCurvePoints.SetNumUninitialized(DisplayedFadeInPixelLenght);
+
+	for (uint32 Pixel = 0; Pixel < DisplayedFadeInPixelLenght; ++Pixel)
+	{
+		const double FadeFraction = (float)Pixel / FadeInPixelLenght;
+		const double CurveValue = Pixel != FadeInPixelLenght - 1 ? 1.f - FMath::Pow(FadeFraction, TrimFadeTransform->StartFadeCurve) : 0.f;
+
+		const uint32 XCoordinate = Pixel + FadeInStartX;
+		FadeInCurvePoints[Pixel] = FVector2D(XCoordinate, CurveValue * AllottedGeometry.Size.Y);
+	}
+
+	const float FadeOutFrames = TrimFadeTransform->EndFadeTime * TransformationWaveInfo.SampleRate;
+	const float FadeOutPixelLength = FadeOutFrames * PixelsPerFrame;
+	FadeOutStartX = FMath::RoundToInt32(FMath::Clamp(EndTimeHandleX - FadeOutPixelLength, StartTimeHandleX, EndTimeHandleX));
+	FadeOutEndX = FMath::RoundToInt32(EndTimeHandleX);
+	
+	const uint32 DisplayedFadeOutPixelLength = FadeOutEndX - FadeOutStartX;
+	FadeOutCurvePoints.SetNumUninitialized(DisplayedFadeOutPixelLength);
+	const uint32 FadeOutPixelOffset = FadeOutPixelLength - DisplayedFadeOutPixelLength;
+
+	for (uint32 Pixel = 0; Pixel < DisplayedFadeOutPixelLength; ++Pixel)
+	{
+		const double FadeFraction = (float)(Pixel + FadeOutPixelOffset) / FadeOutPixelLength;
+		const double CurveValue = Pixel != FadeOutPixelLength - 1 ? FMath::Pow(FadeFraction, TrimFadeTransform->EndFadeCurve) : 1.f;
+
+		const uint32 XCoordinate = Pixel + FadeOutStartX;
+		FadeOutCurvePoints[Pixel] = FVector2D(XCoordinate, CurveValue * AllottedGeometry.Size.Y);
+	}
+
+}
+
+void FWaveformTransformationTrimFadeRenderer::Tick(const FGeometry& AllottedGeometry, const double InCurrentTime, const float InDeltaTime)
+{ 
+	if (!TrimFadeTransform)
+	{
+		return;
+	}
+
+	const float NumFrames = TransformationWaveInfo.NumAvilableSamples / TransformationWaveInfo.NumChannels;
+	const double FirstFrame = FMath::Clamp((TrimFadeTransform->StartTime * TransformationWaveInfo.SampleRate) , 0.f, NumFrames);
+	const double EndFrame = FMath::Clamp((TrimFadeTransform->EndTime * TransformationWaveInfo.SampleRate), FirstFrame, NumFrames);
+
+	PixelsPerFrame = AllottedGeometry.GetLocalSize().X / NumFrames;
+
+	StartTimeHandleX = FirstFrame * PixelsPerFrame;
+	EndTimeHandleX = EndFrame * PixelsPerFrame;
+
+	GenerateFadeCurves(AllottedGeometry);
+}
\ No newline at end of file
diff --git a/Engine/Plugins/Editor/WaveformEditor/Source/WaveformTransformationsWidgets/Private/WaveformTransformationTrimFadeRenderer.h b/Engine/Plugins/Editor/WaveformEditor/Source/WaveformTransformationsWidgets/Private/WaveformTransformationTrimFadeRenderer.h
new file mode 100644
--- /dev/null
+++ b/Engine/Plugins/Editor/WaveformEditor/Source/WaveformTransformationsWidgets/Private/WaveformTransformationTrimFadeRenderer.h
@@ -0,0 +1,43 @@
+// Copyright Epic Games, Inc. All Rights Reserved.
+
+#pragma once
+
+#include "WaveformTransformationRendererBase.h"
+#include "WaveformTransformationTrimFade.h"
+
+class FWaveformTransformationTrimFadeRenderer : public FWaveformTransformationRendererBase
+{
+public:
+	explicit FWaveformTransformationTrimFadeRenderer(const TObjectPtr<UWaveformTransformationTrimFade> TransformationToRender);
+
+	virtual int32 OnPaint(const FPaintArgs& Args, const FGeometry& AllottedGeometry, const FSlateRect& MyCullingRect, FSlateWindowElementList& OutDrawElements, int32 LayerId, const FWidgetStyle& InWidgetStyle, bool bParentEnabled) const override;
+	virtual void Tick(const FGeometry& AllottedGeometry, const double InCurrentTime, const float InDeltaTime) override;
+
+private:
+	int32 DrawTrimHandles(const FGeometry& AllottedGeometry, FSlateWindowElementList& OutDrawElements, int32 LayerId) const;
+	int32 DrawFadeCurves(const FGeometry& AllottedGeometry, FSlateWindowElementList& OutDrawElements, int32 LayerId) const;
+	void GenerateFadeCurves(const FGeometry& AllottedGeometry);
+
+	TObjectPtr<UWaveformTransformationTrimFade> TrimFadeTransform = nullptr;
+
+	float StartTimeHandleX = 0.f;
+	float EndTimeHandleX = 0.f;
+	uint32 FadeInStartX = 0;
+	uint32 FadeInEndX = 0;
+	uint32 FadeOutStartX = 0;
+	uint32 FadeOutEndX = 0;
+	TArray<FVector2D> FadeInCurvePoints;
+	TArray<FVector2D> FadeOutCurvePoints;
+
+	double PixelsPerFrame = 0.0;
+
+	bool bScrubbingLeftHandle = false;
+	bool bScrubbingRightHandle = false;
+	bool bScrubbingFadeIn = false;
+	bool bScrubbingFadeOut = false;
+
+	
+	const float InteractionPixelXDelta = 10;
+	const float InteractionRatioYDelta = 0.07f;
+	const float MouseWheelStep = 0.25f;
+};
\ No newline at end of file
diff --git a/Engine/Plugins/Editor/WaveformEditor/Source/WaveformTransformationsWidgets/Private/WaveformTransformationsRenderManager.cpp b/Engine/Plugins/Editor/WaveformEditor/Source/WaveformTransformationsWidgets/Private/WaveformTransformationsRenderManager.cpp
--- a/Engine/Plugins/Editor/WaveformEditor/Source/WaveformTransformationsWidgets/Private/WaveformTransformationsRenderManager.cpp
+++ b/Engine/Plugins/Editor/WaveformEditor/Source/WaveformTransformationsWidgets/Private/WaveformTransformationsRenderManager.cpp
@@ -4,13 +4,14 @@
 
 #include "DSP/BufferVectorOperations.h"
 #include "DSP/FloatArrayMath.h"
+#include "IWaveformTransformationRenderer.h"
 #include "Sound/SoundWave.h"
-#include "SWaveformTransformationRenderLayer.h"
-#include "Widgets/SWidget.h"
 
-FWaveformTransformationsRenderManager::FWaveformTransformationsRenderManager(TObjectPtr<USoundWave> InSoundWave, TSharedRef<FWaveformEditorRenderData> InWaveformRenderData, TSharedRef<FWaveformEditorTransportCoordinator> InTransportCoordinator, TSharedRef<FWaveformEditorZoomController> InZoomController)
+FWaveformTransformationsRenderManager::FWaveformTransformationsRenderManager(TObjectPtr<USoundWave> InSoundWave, 
+	TSharedRef<FWaveformEditorRenderData> InWaveformRenderData, 
+	TFunction<void(FPropertyChangedEvent&, FEditPropertyChain*)> InPropertyChangeNotifier)
 	: SoundWaveToRender(InSoundWave)
-	, LayersFactory(MakeUnique<FWaveformTransformationRenderLayerFactory>(InWaveformRenderData, InTransportCoordinator, InZoomController))
+	, LayersFactory(MakeUnique<FWaveformTransformationRenderLayerFactory>(InWaveformRenderData, InPropertyChangeNotifier))
 {
 	GenerateLayersChain();	
 }
@@ -30,8 +31,9 @@ void FWaveformTransformationsRenderManager::GenerateLayersChain()
 			{
 				TransformationsToRender.Add(Transformation);
 				
-				TSharedPtr<SWaveformTransformationRenderLayer> TransformationUI = LayersFactory->Create(Transformation);
-				RenderLayers.Add(TransformationUI);
+				TSharedPtr<IWaveformTransformationRenderer> TransformationUI = LayersFactory->Create(Transformation);
+				FTransformationRenderLayerInfo RenderLayerInfo = FTransformationRenderLayerInfo(TransformationUI, FTransformationLayerConstraints(0.f, 1.f));
+				RenderLayers.Add(RenderLayerInfo);
 				
 			}
 		}
@@ -40,10 +42,7 @@ void FWaveformTransformationsRenderManager::GenerateLayersChain()
 
 	}
 
-	TSharedPtr<SWidget>* WidgetPtr = reinterpret_cast<TSharedPtr<SWidget>*>(RenderLayers.GetData());
-	RenderLayersWidgetView = MakeArrayView(WidgetPtr, RenderLayers.Num());
-
-	OnLayersChainGenerated.Broadcast(RenderLayersWidgetView.GetData(), RenderLayers.Num());
+	OnLayersChainGenerated.Broadcast(RenderLayers.GetData(), RenderLayers.Num());
 }
 
 void FWaveformTransformationsRenderManager::UpdateRenderElements()
@@ -51,8 +50,6 @@ void FWaveformTransformationsRenderManager::UpdateRenderElements()
 	GenerateRenderDataInternal();
 }
 
-
-
 void FWaveformTransformationsRenderManager::GenerateRenderDataInternal()
 {
 	check(SoundWaveToRender);
@@ -95,15 +92,21 @@ void FWaveformTransformationsRenderManager::GenerateRenderDataInternal()
 		{
 			Transformations[i]->ProcessAudio(TransformationInfo);
 
-			FirstEditedSample += TransformationInfo.StartFrameOffset;
-			LastEditedSample = TransformationInfo.NumEditedSamples == 0 ? LastEditedSample : FirstEditedSample + TransformationInfo.NumEditedSamples;
-
-			if (RenderLayers[i])
+			if (RenderLayers[i].Key)
 			{
-				FWaveformTransformationRenderLayerInfo RenderLayerInfo{ TransformationInfo.SampleRate, TransformationInfo.NumChannels, FirstEditedSample, TransformationInfo.NumEditedSamples };
-				RenderLayers[i]->SetTransformationWaveInfo(MoveTemp(RenderLayerInfo));
+				FWaveformTransformationRenderInfo RenderLayerInfo{ TransformationInfo.SampleRate, TransformationInfo.NumChannels, FirstEditedSample, LastEditedSample - FirstEditedSample };
+				RenderLayers[i].Key->SetTransformationWaveInfo(MoveTemp(RenderLayerInfo));
+
+				float LeftWidgetConstraint = FirstEditedSample / (float) NumOriginalSamples;
+				float RightWidgetConstraint = LastEditedSample / (float) NumOriginalSamples;
+
+				RenderLayers[i].Value = FTransformationLayerConstraints(LeftWidgetConstraint, RightWidgetConstraint);
+
 			}
 
+			FirstEditedSample += TransformationInfo.StartFrameOffset;
+			LastEditedSample = TransformationInfo.NumEditedSamples <= 0 ? LastEditedSample : FirstEditedSample + TransformationInfo.NumEditedSamples;
+
 			if (bChainChangesFileLength)
 			{
 				check(LastEditedSample > FirstEditedSample);
@@ -116,7 +119,7 @@ void FWaveformTransformationsRenderManager::GenerateRenderDataInternal()
 		}
 
 		check(DurationHiglightLayer)
-		FWaveformTransformationRenderLayerInfo DurationLayerInfo{ TransformationInfo.SampleRate, TransformationInfo.NumChannels, FirstEditedSample, LastEditedSample - FirstEditedSample };
+		FWaveformTransformationRenderInfo DurationLayerInfo{ TransformationInfo.SampleRate, TransformationInfo.NumChannels, FirstEditedSample, LastEditedSample - FirstEditedSample };
 		DurationHiglightLayer->SetTransformationWaveInfo(MoveTemp(DurationLayerInfo));
 			
 		if (!bChainChangesFileLength)
@@ -142,6 +145,7 @@ void FWaveformTransformationsRenderManager::GenerateRenderDataInternal()
 	}
 
 	OnRenderDataGenerated.Broadcast(RawPCMData.GetData(), NumOriginalSamples, FirstEditedSample, LastEditedSample, SampleRate, NumChannels);
+	OnRenderElementsUpdated.Broadcast();
 }
 
 TArray<Audio::FTransformationPtr> FWaveformTransformationsRenderManager::CreateTransformations() const
@@ -168,9 +172,9 @@ const bool FWaveformTransformationsRenderManager::CanChainChangeFileLength(const
 	return bCanChainChangeFileLength;
 }
 
-TArrayView<TSharedPtr<SWidget>> FWaveformTransformationsRenderManager::GetTransformLayers() const
+TArrayView<const FTransformationRenderLayerInfo> FWaveformTransformationsRenderManager::GetTransformLayers() const
 {
-	return RenderLayersWidgetView;
+	return MakeArrayView(RenderLayers.GetData(), RenderLayers.Num());
 }
 
 void FWaveformTransformationsRenderManager::CreateDurationHighlightLayer()
@@ -180,6 +184,6 @@ void FWaveformTransformationsRenderManager::CreateDurationHighlightLayer()
 		DurationHiglightLayer = LayersFactory->CreateDurationHiglightLayer();
 	}
 
-	RenderLayers.Add(DurationHiglightLayer);
-	DurationHiglightLayer->SetTransformationWaveInfo(FWaveformTransformationRenderLayerInfo());
+	DurationHiglightLayer->SetTransformationWaveInfo(FWaveformTransformationRenderInfo());
+	RenderLayers.Emplace(DurationHiglightLayer, FTransformationLayerConstraints(0.f,1.f));
 }
\ No newline at end of file
diff --git a/Engine/Plugins/Editor/WaveformEditor/Source/WaveformTransformationsWidgets/Public/SWaveformTransformationRenderLayer.h b/Engine/Plugins/Editor/WaveformEditor/Source/WaveformTransformationsWidgets/Public/SWaveformTransformationRenderLayer.h
deleted file mode 100644
diff --git a/Engine/Plugins/Editor/WaveformEditor/Source/WaveformTransformationsWidgets/Public/WaveformTransformationRenderLayerFactory.h b/Engine/Plugins/Editor/WaveformEditor/Source/WaveformTransformationsWidgets/Public/WaveformTransformationRenderLayerFactory.h
--- a/Engine/Plugins/Editor/WaveformEditor/Source/WaveformTransformationsWidgets/Public/WaveformTransformationRenderLayerFactory.h
+++ b/Engine/Plugins/Editor/WaveformEditor/Source/WaveformTransformationsWidgets/Public/WaveformTransformationRenderLayerFactory.h
@@ -7,8 +7,7 @@
 
 class FWaveformEditorRenderData;
 class FWaveformEditorTransportCoordinator;
-class FWaveformEditorZoomController;
-class SWaveformTransformationRenderLayer;
+class IWaveformTransformationRenderer;
 class UWaveformTransformationBase;
 
 class FWaveformTransformationRenderLayerFactory
@@ -16,16 +15,15 @@ class FWaveformTransformationRenderLayerFactory
 public:
 	explicit FWaveformTransformationRenderLayerFactory(
 		TSharedRef<FWaveformEditorRenderData> InWaveformRenderData, 
-		TSharedRef<FWaveformEditorTransportCoordinator> InTransportCoordinator, 
-		TSharedRef<FWaveformEditorZoomController> InZoomController);
+		TFunction<void(FPropertyChangedEvent&, FEditPropertyChain*)> InTransformationChangeNotifier);
 
 	~FWaveformTransformationRenderLayerFactory() = default;
 
-	TSharedPtr<SWaveformTransformationRenderLayer> Create(TObjectPtr<UWaveformTransformationBase> InTransformationToRender) const;
-	TSharedPtr<SWaveformTransformationRenderLayer> CreateDurationHiglightLayer () const;
+	TSharedPtr<IWaveformTransformationRenderer> Create(TObjectPtr<UWaveformTransformationBase> InTransformationToRender) const;
+	TSharedPtr<IWaveformTransformationRenderer> CreateDurationHiglightLayer () const;
 
 private:
 	TSharedPtr<FWaveformEditorRenderData> WaveformRenderData = nullptr;
 	TSharedPtr<FWaveformEditorTransportCoordinator> TransportCoordinator = nullptr;
-	TSharedPtr<FWaveformEditorZoomController> ZoomController = nullptr;
+	TFunction<void(FPropertyChangedEvent&, FEditPropertyChain*)> TransformationChangeNotifier;
 };
\ No newline at end of file
diff --git a/Engine/Plugins/Editor/WaveformEditor/Source/WaveformTransformationsWidgets/Public/WaveformTransformationsRenderManager.h b/Engine/Plugins/Editor/WaveformEditor/Source/WaveformTransformationsWidgets/Public/WaveformTransformationsRenderManager.h
--- a/Engine/Plugins/Editor/WaveformEditor/Source/WaveformTransformationsWidgets/Public/WaveformTransformationsRenderManager.h
+++ b/Engine/Plugins/Editor/WaveformEditor/Source/WaveformTransformationsWidgets/Public/WaveformTransformationsRenderManager.h
@@ -6,12 +6,15 @@
 #include "IWaveformTransformation.h"
 #include "WaveformTransformationRenderLayerFactory.h"
 
+class IWaveformTransformationRenderer;
 class USoundWave;
 class UWaveformTransformationBase;
-class SWaveformTransformationRenderLayer;
-class SWidget;
 
-DECLARE_MULTICAST_DELEGATE_TwoParams(FOnLayersChainGenerated, TSharedPtr<SWidget>* /*First Layer ptr*/, const int32 /* NLayers */)
+using FTransformationLayerConstraints = TPair<float, float>;
+using FTransformationRenderLayerInfo = TPair<TSharedPtr<IWaveformTransformationRenderer>, FTransformationLayerConstraints>;
+
+DECLARE_MULTICAST_DELEGATE(FOnRenderElementsUpdated)
+DECLARE_MULTICAST_DELEGATE_TwoParams(FOnLayersChainGenerated, FTransformationRenderLayerInfo* /*First Layer ptr*/, const int32 /* NLayers */)
 DECLARE_MULTICAST_DELEGATE_SixParams(FOnRenderDataGenerated, const uint8* /* First RawPCMData element */, const uint32 /*Num Samples*/, const uint32 /*First Edited Sample*/, const uint32 /*Last Edited Sample*/, const uint32/*Sample Rate*/, const uint16 /*Num Channels*/)
 
 
@@ -22,8 +25,8 @@ DECLARE_MULTICAST_DELEGATE_SixParams(FOnRenderDataGenerated, const uint8* /* Fir
 /*  																				 */
 /* The main UI elements created are:												 */
 /*  Render Data : a uint8 array containing the transformed samples of the waveform   */ 
-/*  Transform Layers : an array of Widgets containing the UI for each transformation */ 
-/*  in the chain																	 */
+/*  Transform Layers : an array of UI renderers for each transformation				 */ 
+/*  in the chain. Renderers should be contained in dummy widgets. 					 */
 /* 																					 */
 /* These are created so that the entire stack of transformation is displayable.		 */
 /* E.g : if a 10s long file is trimmed from second 2, seconds 0 to 1 will 			 */
@@ -44,11 +47,10 @@ public:
 	explicit FWaveformTransformationsRenderManager(
 		TObjectPtr<USoundWave> InSoundWave, 
 		TSharedRef<FWaveformEditorRenderData> InWaveformRenderData, 
-		TSharedRef<FWaveformEditorTransportCoordinator> InTransportCoordinator, 
-		TSharedRef<FWaveformEditorZoomController> InZoomController
+		TFunction<void(FPropertyChangedEvent&, FEditPropertyChain*)> InPropertyChangeNotifier
 	);
 
-	/** Used to generate the stack of transformations UI widgets			*/
+	/** Used to generate the stack of transformations UI					*/
 	/** Should be called when the waveform transformation chain is changed	*/
 	void GenerateLayersChain();
 
@@ -56,14 +58,17 @@ public:
 	/** Should be called when the transformations parameters are changed				 */
 	void UpdateRenderElements();
 
-	TArrayView<TSharedPtr<SWidget>> GetTransformLayers() const;
+	TArrayView<const FTransformationRenderLayerInfo> GetTransformLayers() const;
 
-	/** Called when a new layer chain of transformations UI is generated */
+	/** Called when a new layer chain of transformations UI is created */
 	FOnLayersChainGenerated OnLayersChainGenerated;
 
 	/** Called when new render data is generated */
 	FOnRenderDataGenerated OnRenderDataGenerated;
 
+	/** Called when the different render elements are updated */
+	FOnRenderElementsUpdated OnRenderElementsUpdated;
+
 private:
 	void GenerateRenderDataInternal();
 
@@ -73,13 +78,12 @@ private:
 	const bool CanChainChangeFileLength(const TArray<Audio::FTransformationPtr>& TransformationArray) const;
 
 	TArray<TObjectPtr<UWaveformTransformationBase>> TransformationsToRender;
-	TArray<TSharedPtr<SWaveformTransformationRenderLayer>> RenderLayers;
-	TArrayView<TSharedPtr<SWidget>> RenderLayersWidgetView;
+	TArray<FTransformationRenderLayerInfo> RenderLayers;
 
 	TObjectPtr<USoundWave> SoundWaveToRender = nullptr;
 
 	TArray<uint8> RawPCMData;
 
 	TUniquePtr<FWaveformTransformationRenderLayerFactory> LayersFactory = nullptr;
-	TSharedPtr<SWaveformTransformationRenderLayer> DurationHiglightLayer = nullptr;
+	TSharedPtr<IWaveformTransformationRenderer> DurationHiglightLayer = nullptr;
 };
\ No newline at end of file
