[WaveformEditor] Add transformed duration getters to render data #trivial #rb none #preflight 630e177098103d4521a6d09e

diff --git a/Engine/Plugins/Editor/WaveformEditor/Source/WaveformEditorWidgets/Private/WaveformEditorRenderData.cpp b/Engine/Plugins/Editor/WaveformEditor/Source/WaveformEditorWidgets/Private/WaveformEditorRenderData.cpp
--- a/Engine/Plugins/Editor/WaveformEditor/Source/WaveformEditorWidgets/Private/WaveformEditorRenderData.cpp
+++ b/Engine/Plugins/Editor/WaveformEditor/Source/WaveformEditorWidgets/Private/WaveformEditorRenderData.cpp
@@ -32,11 +32,21 @@ const float FWaveformEditorRenderData::GetOriginalWaveformDurationInSeconds() co
 	return OriginalWaveformDurationInSeconds;
 }
 
+const float FWaveformEditorRenderData::GetTransformedWaveformDurationInSeconds() const
+{
+	return TransformedWaveformDurationInSeconds;
+}
+
 const uint32 FWaveformEditorRenderData::GetOriginalWaveformFrames() const
 {
 	return OriginalWaveformFrames;
 }
 
+const uint32 FWaveformEditorRenderData::GetTransformedWaveformFrames() const
+{
+	return TransformedWaveformFrames;
+}
+
 const TRange<float> FWaveformEditorRenderData::GetTransformedWaveformBounds() const
 {
 	return TransformedWaveformBounds;
@@ -51,8 +61,11 @@ void FWaveformEditorRenderData::UpdateRenderData(const uint8* InRawData, const u
 	const int16* SampleDataPtr = reinterpret_cast<const int16*>(InRawData);
 	SampleData = MakeArrayView(SampleDataPtr, NumSamples);
 
-	OriginalWaveformDurationInSeconds = NumSamples / NumChannels / SampleRate;
 	OriginalWaveformFrames = NumSamples / NumChannels;
+	OriginalWaveformDurationInSeconds = OriginalWaveformFrames / SampleRate;
+	
+	TransformedWaveformFrames = (InLastEditedSample - InFirstEditedSample) / NumChannels;
+	TransformedWaveformDurationInSeconds = TransformedWaveformFrames / SampleRate;
 
 	const float FirstEditedSampleOffset = InFirstEditedSample / (float) NumSamples;
 	const float LastEditedSampleOffset = (NumSamples - InLastEditedSample) / (float) NumSamples;	
diff --git a/Engine/Plugins/Editor/WaveformEditor/Source/WaveformEditorWidgets/Public/WaveformEditorRenderData.h b/Engine/Plugins/Editor/WaveformEditor/Source/WaveformEditorWidgets/Public/WaveformEditorRenderData.h
--- a/Engine/Plugins/Editor/WaveformEditor/Source/WaveformEditorWidgets/Public/WaveformEditorRenderData.h
+++ b/Engine/Plugins/Editor/WaveformEditor/Source/WaveformEditorWidgets/Public/WaveformEditorRenderData.h
@@ -33,7 +33,9 @@ public:
 	const uint32 GetNumSamples() const;
 	const uint32 GetSampleRate() const;
 	const float GetOriginalWaveformDurationInSeconds() const;
+	const float GetTransformedWaveformDurationInSeconds() const;
 	const uint32 GetOriginalWaveformFrames() const;
+	const uint32 GetTransformedWaveformFrames() const;
 	const TRange<float> GetTransformedWaveformBounds() const;
 
 	/** Call this to generate the view and the info shown by the render data */
@@ -51,7 +53,9 @@ private:
 	uint32 SampleRate = 0;
 
 	float OriginalWaveformDurationInSeconds = 0.f;
+	float TransformedWaveformDurationInSeconds = 0.f;
 	uint32 OriginalWaveformFrames = 0;
+	uint32 TransformedWaveformFrames = 0;
 
 	/* The bounds of the transformed waveform in relation to the original */
 	TRange<float> TransformedWaveformBounds = TRange<float>::Inclusive(0.f, 1.f);
