[WaveformEditor] fix implicit integer division in FWaveformEditorRenderData::UpdateRenderData
#trivial
fix for #jira UE-161604

diff --git a/Engine/Plugins/Editor/WaveformEditor/Source/WaveformEditorWidgets/Private/WaveformEditorRenderData.cpp b/Engine/Plugins/Editor/WaveformEditor/Source/WaveformEditorWidgets/Private/WaveformEditorRenderData.cpp
--- a/Engine/Plugins/Editor/WaveformEditor/Source/WaveformEditorWidgets/Private/WaveformEditorRenderData.cpp
+++ b/Engine/Plugins/Editor/WaveformEditor/Source/WaveformEditorWidgets/Private/WaveformEditorRenderData.cpp
@@ -62,10 +62,10 @@ void FWaveformEditorRenderData::UpdateRenderData(const uint8* InRawData, const u
 	SampleData = MakeArrayView(SampleDataPtr, NumSamples);
 
 	OriginalWaveformFrames = NumSamples / NumChannels;
-	OriginalWaveformDurationInSeconds = OriginalWaveformFrames / SampleRate;
+	OriginalWaveformDurationInSeconds = OriginalWaveformFrames / (float) SampleRate;
 	
 	TransformedWaveformFrames = (InLastEditedSample - InFirstEditedSample) / NumChannels;
-	TransformedWaveformDurationInSeconds = TransformedWaveformFrames / SampleRate;
+	TransformedWaveformDurationInSeconds = TransformedWaveformFrames / (float) SampleRate;
 
 	const float FirstEditedSampleOffset = InFirstEditedSample / (float) NumSamples;
 	const float LastEditedSampleOffset = (NumSamples - InLastEditedSample) / (float) NumSamples;	
