[WaveformEditor] Advanced Trim Fade UI with interactions

diff --git a/Engine/Plugins/Editor/WaveformEditor/Source/WaveformTransformationsWidgets/Private/WaveformTransformationTrimFadeRenderer.cpp b/Engine/Plugins/Editor/WaveformEditor/Source/WaveformTransformationsWidgets/Private/WaveformTransformationTrimFadeRenderer.cpp
--- a/Engine/Plugins/Editor/WaveformEditor/Source/WaveformTransformationsWidgets/Private/WaveformTransformationTrimFadeRenderer.cpp
+++ b/Engine/Plugins/Editor/WaveformEditor/Source/WaveformTransformationsWidgets/Private/WaveformTransformationTrimFadeRenderer.cpp
@@ -126,7 +126,64 @@ void FWaveformTransformationTrimFadeRenderer::GenerateFadeCurves(const FGeometry
 		const uint32 XCoordinate = Pixel + FadeOutStartX;
 		FadeOutCurvePoints[Pixel] = FVector2D(XCoordinate, CurveValue * AllottedGeometry.Size.Y);
 	}
+}
+
+FCursorReply FWaveformTransformationTrimFadeRenderer::OnCursorQuery(const FGeometry& MyGeometry, const FPointerEvent& CursorEvent) const
+{
+	const FVector2D LocalCursorPosition = GetLocalCursorPosition(CursorEvent, MyGeometry);
+
+	if (TrimFadeInteractionType == ETrimFadeInteractionType::ScrubbingFadeIn || TrimFadeInteractionType == ETrimFadeInteractionType::ScrubbingFadeOut)
+	{
+		return FCursorReply::Cursor(EMouseCursor::GrabHandClosed);
+	}
+
+	if (IsCursorInFadeInInteractionRange(LocalCursorPosition, MyGeometry) || IsCursorInFadeOutInteractionRange(LocalCursorPosition, MyGeometry))
+	{
+		return FCursorReply::Cursor(EMouseCursor::GrabHand);
+	}
+
+	if (TrimFadeInteractionType == ETrimFadeInteractionType::ScrubbingLeftHandle || TrimFadeInteractionType == ETrimFadeInteractionType::ScrubbingRightHandle || StartTimeInteractionXRange.Contains(LocalCursorPosition.X) || EndTimeInteractionXRange.Contains(LocalCursorPosition.X))
+	{
+		return FCursorReply::Cursor(EMouseCursor::ResizeLeftRight);
+	}
+
+	return FCursorReply::Unhandled();
+}
+
+FReply FWaveformTransformationTrimFadeRenderer::OnMouseWheel(SWidget& OwnerWidget, const FGeometry& MyGeometry, const FPointerEvent& MouseEvent)
+{
+	const FVector2D LocalCursorPosition = GetLocalCursorPosition(MouseEvent, MyGeometry);
+
+	if (IsCursorInFadeInInteractionRange(LocalCursorPosition, MyGeometry))
+	{		
+		const float FadeCurveDelta = MouseEvent.GetWheelDelta() * MouseWheelStep;
+		TrimFadeTransform->StartFadeCurve = FMath::Clamp(TrimFadeTransform->StartFadeCurve + FadeCurveDelta, 0.f, 10.f);
+		NotifyTransformationPropertyChanged(TrimFadeTransform, GET_MEMBER_NAME_CHECKED(UWaveformTransformationTrimFade, StartFadeCurve), EPropertyChangeType::ValueSet);
+		return FReply::Handled();
+	}
+
+	if (IsCursorInFadeOutInteractionRange(LocalCursorPosition, MyGeometry))
+	{
+		const float FadeCurveDelta = MouseEvent.GetWheelDelta() * MouseWheelStep;
+		TrimFadeTransform->EndFadeCurve = FMath::Clamp(TrimFadeTransform->EndFadeCurve + FadeCurveDelta, 0.f, 10.f);
+		NotifyTransformationPropertyChanged(TrimFadeTransform, GET_MEMBER_NAME_CHECKED(UWaveformTransformationTrimFade, EndFadeCurve), EPropertyChangeType::ValueSet);
+		return FReply::Handled();
+	}
+
+	return FReply::Unhandled();
+}
+
+FVector2D FWaveformTransformationTrimFadeRenderer::GetLocalCursorPosition(const FPointerEvent& MouseEvent, const FGeometry& EventGeometry) const
+{
+	const FVector2D ScreenSpacePosition = MouseEvent.GetScreenSpacePosition();
+	return  EventGeometry.AbsoluteToLocal(ScreenSpacePosition);
+}
 
+float FWaveformTransformationTrimFadeRenderer::ConvertXRatioToTime(const float InRatio) const
+{
+	const float NumFrames = TransformationWaveInfo.NumAvilableSamples / TransformationWaveInfo.NumChannels;
+	const float FrameSelected = NumFrames * InRatio;
+	return FrameSelected * TransformationWaveInfo.NumChannels / TransformationWaveInfo.SampleRate;
 }
 
 void FWaveformTransformationTrimFadeRenderer::Tick(const FGeometry& AllottedGeometry, const double InCurrentTime, const float InDeltaTime)
@@ -146,4 +203,123 @@ void FWaveformTransformationTrimFadeRenderer::Tick(const FGeometry& AllottedGeom
 	EndTimeHandleX = EndFrame * PixelsPerFrame;
 
 	GenerateFadeCurves(AllottedGeometry);
+	UpdateInteractionRange();
+}
+
+void FWaveformTransformationTrimFadeRenderer::UpdateInteractionRange()
+{
+	StartTimeInteractionXRange.SetLowerBoundValue(StartTimeHandleX - InteractionPixelXDelta);
+	StartTimeInteractionXRange.SetUpperBoundValue(StartTimeHandleX + InteractionPixelXDelta);
+	EndTimeInteractionXRange.SetLowerBoundValue(EndTimeHandleX - InteractionPixelXDelta);
+	EndTimeInteractionXRange.SetUpperBoundValue(EndTimeHandleX + InteractionPixelXDelta);
+	FadeInInteractionXRange.SetLowerBoundValue(FadeInEndX - InteractionPixelXDelta);
+	FadeInInteractionXRange.SetUpperBoundValue(FadeInEndX + InteractionPixelXDelta);
+	FadeOutInteractionXRange.SetLowerBoundValue(FadeOutStartX - InteractionPixelXDelta);
+	FadeOutInteractionXRange.SetUpperBoundValue(FadeOutStartX + InteractionPixelXDelta);
+}
+
+FReply FWaveformTransformationTrimFadeRenderer::OnMouseButtonUp(SWidget& OwnerWidget, const FGeometry& MyGeometry, const FPointerEvent& MouseEvent)
+{
+	if (TrimFadeInteractionType != ETrimFadeInteractionType::None)
+	{
+		TrimFadeInteractionType = ETrimFadeInteractionType::None;
+		return FReply::Handled().ReleaseMouseCapture();
+	}
+
+	return FReply::Unhandled();
+}
+
+FReply FWaveformTransformationTrimFadeRenderer::OnMouseButtonDown(SWidget& OwnerWidget, const FGeometry& MyGeometry, const FPointerEvent& MouseEvent)
+{
+	const FVector2D LocalCursorPosition = GetLocalCursorPosition(MouseEvent, MyGeometry);
+
+	const bool bIsLeftMouseButton = MouseEvent.IsMouseButtonDown(EKeys::LeftMouseButton);
+	
+	if (bIsLeftMouseButton)
+	{
+		TrimFadeInteractionType = GetInteractionTypeFromCursorPosition(LocalCursorPosition, MyGeometry);
+
+		if (TrimFadeInteractionType != ETrimFadeInteractionType::None)
+		{
+			return FReply::Handled().CaptureMouse(OwnerWidget.AsShared()).PreventThrottling();
+		}
+	}
+
+	return FReply::Unhandled();
+}
+
+
+FReply FWaveformTransformationTrimFadeRenderer::OnMouseMove(SWidget& OwnerWidget, const FGeometry& MyGeometry, const FPointerEvent& MouseEvent)
+{
+	if (MouseEvent.IsMouseButtonDown(EKeys::LeftMouseButton) && TrimFadeInteractionType != ETrimFadeInteractionType::None)
+	{
+		const FVector2D LocalCursorPosition = GetLocalCursorPosition(MouseEvent, MyGeometry);
+		const float LocalCursorXRatio = FMath::Clamp(LocalCursorPosition.X / MyGeometry.GetLocalSize().X, 0.f, 1.f);
+		const float SelectedTime = ConvertXRatioToTime(LocalCursorXRatio);
+
+		switch (TrimFadeInteractionType)
+		{
+		case FWaveformTransformationTrimFadeRenderer::ETrimFadeInteractionType::None:
+			break;
+		case FWaveformTransformationTrimFadeRenderer::ETrimFadeInteractionType::ScrubbingLeftHandle:
+			TrimFadeTransform->StartTime = SelectedTime;
+			NotifyTransformationPropertyChanged(TrimFadeTransform, GET_MEMBER_NAME_CHECKED(UWaveformTransformationTrimFade, StartTime), EPropertyChangeType::ValueSet);
+			break;
+		case FWaveformTransformationTrimFadeRenderer::ETrimFadeInteractionType::ScrubbingRightHandle:
+			TrimFadeTransform->EndTime = SelectedTime;
+			NotifyTransformationPropertyChanged(TrimFadeTransform, GET_MEMBER_NAME_CHECKED(UWaveformTransformationTrimFade, EndTime), EPropertyChangeType::ValueSet);
+			break;
+		case FWaveformTransformationTrimFadeRenderer::ETrimFadeInteractionType::ScrubbingFadeIn:
+			TrimFadeTransform->StartFadeTime = FMath::Clamp(SelectedTime - TrimFadeTransform->StartTime, 0.f, TNumericLimits<float>().Max());
+			NotifyTransformationPropertyChanged(TrimFadeTransform, GET_MEMBER_NAME_CHECKED(UWaveformTransformationTrimFade, StartFadeTime), EPropertyChangeType::ValueSet);
+			break;
+		case FWaveformTransformationTrimFadeRenderer::ETrimFadeInteractionType::ScrubbingFadeOut:
+			TrimFadeTransform->EndFadeTime = FMath::Clamp(TrimFadeTransform->EndTime - SelectedTime, 0.f, TNumericLimits<float>().Max());
+			NotifyTransformationPropertyChanged(TrimFadeTransform, GET_MEMBER_NAME_CHECKED(UWaveformTransformationTrimFade, EndFadeTime), EPropertyChangeType::ValueSet);
+			break;
+		default:
+			break;
+		}
+
+		return FReply::Handled().CaptureMouse(OwnerWidget.AsShared());
+	}
+
+	return FReply::Unhandled();
+}
+
+FWaveformTransformationTrimFadeRenderer::ETrimFadeInteractionType FWaveformTransformationTrimFadeRenderer::GetInteractionTypeFromCursorPosition(const FVector2D& InLocalCursorPosition, const FGeometry& WidgetGeometry) const
+{
+	if (IsCursorInFadeInInteractionRange(InLocalCursorPosition, WidgetGeometry))
+	{
+		return ETrimFadeInteractionType::ScrubbingFadeIn;
+	}
+
+	if (IsCursorInFadeOutInteractionRange(InLocalCursorPosition, WidgetGeometry))
+	{
+		return ETrimFadeInteractionType::ScrubbingFadeOut;
+	}
+
+	if (StartTimeInteractionXRange.Contains(InLocalCursorPosition.X))
+	{
+		return ETrimFadeInteractionType::ScrubbingLeftHandle;
+	}
+
+	if (EndTimeInteractionXRange.Contains(InLocalCursorPosition.X))
+	{
+		return ETrimFadeInteractionType::ScrubbingRightHandle;
+	}
+
+	return ETrimFadeInteractionType::None;
+}
+
+bool FWaveformTransformationTrimFadeRenderer::IsCursorInFadeInInteractionRange(const FVector2D& InLocalCursorPosition, const FGeometry& WidgetGeometry) const
+{
+	return FadeInInteractionXRange.Contains(InLocalCursorPosition.X)
+		&& InLocalCursorPosition.Y < WidgetGeometry.GetLocalSize().Y* InteractionRatioYDelta;
+}
+
+bool FWaveformTransformationTrimFadeRenderer::IsCursorInFadeOutInteractionRange(const FVector2D& InLocalCursorPosition, const FGeometry& WidgetGeometry) const
+{
+	return FadeOutInteractionXRange.Contains(InLocalCursorPosition.X)
+		&& InLocalCursorPosition.Y < WidgetGeometry.GetLocalSize().Y* InteractionRatioYDelta;
 }
\ No newline at end of file
diff --git a/Engine/Plugins/Editor/WaveformEditor/Source/WaveformTransformationsWidgets/Private/WaveformTransformationTrimFadeRenderer.h b/Engine/Plugins/Editor/WaveformEditor/Source/WaveformTransformationsWidgets/Private/WaveformTransformationTrimFadeRenderer.h
--- a/Engine/Plugins/Editor/WaveformEditor/Source/WaveformTransformationsWidgets/Private/WaveformTransformationTrimFadeRenderer.h
+++ b/Engine/Plugins/Editor/WaveformEditor/Source/WaveformTransformationsWidgets/Private/WaveformTransformationTrimFadeRenderer.h
@@ -12,11 +12,23 @@ public:
 
 	virtual int32 OnPaint(const FPaintArgs& Args, const FGeometry& AllottedGeometry, const FSlateRect& MyCullingRect, FSlateWindowElementList& OutDrawElements, int32 LayerId, const FWidgetStyle& InWidgetStyle, bool bParentEnabled) const override;
 	virtual void Tick(const FGeometry& AllottedGeometry, const double InCurrentTime, const float InDeltaTime) override;
+	virtual FReply OnMouseButtonDown(SWidget& OwnerWidget, const FGeometry& MyGeometry, const FPointerEvent& MouseEvent) override;
+	virtual FReply OnMouseButtonUp(SWidget& OwnerWidget, const FGeometry& MyGeometry, const FPointerEvent& MouseEvent) override;
+	virtual FReply OnMouseMove(SWidget& OwnerWidget, const FGeometry& MyGeometry, const FPointerEvent& MouseEvent) override;
+	virtual FReply OnMouseWheel(SWidget& OwnerWidget, const FGeometry& MyGeometry, const FPointerEvent& MouseEvent) override;
+	virtual FCursorReply OnCursorQuery(const FGeometry& MyGeometry, const FPointerEvent& CursorEvent) const override;
 
 private:
 	int32 DrawTrimHandles(const FGeometry& AllottedGeometry, FSlateWindowElementList& OutDrawElements, int32 LayerId) const;
 	int32 DrawFadeCurves(const FGeometry& AllottedGeometry, FSlateWindowElementList& OutDrawElements, int32 LayerId) const;
 	void GenerateFadeCurves(const FGeometry& AllottedGeometry);
+	void UpdateInteractionRange();
+
+	bool IsCursorInFadeInInteractionRange(const FVector2D& InLocalCursorPosition, const FGeometry& WidgetGeometry) const;
+	bool IsCursorInFadeOutInteractionRange(const FVector2D& InLocalCursorPosition, const FGeometry& WidgetGeometry) const;
+
+	FVector2D GetLocalCursorPosition(const FPointerEvent& MouseEvent, const FGeometry& EventGeometry) const;
+	float ConvertXRatioToTime(const float InRatio) const;
 
 	TObjectPtr<UWaveformTransformationTrimFade> TrimFadeTransform = nullptr;
 
@@ -29,15 +41,25 @@ private:
 	TArray<FVector2D> FadeInCurvePoints;
 	TArray<FVector2D> FadeOutCurvePoints;
 
-	double PixelsPerFrame = 0.0;
-
-	bool bScrubbingLeftHandle = false;
-	bool bScrubbingRightHandle = false;
-	bool bScrubbingFadeIn = false;
-	bool bScrubbingFadeOut = false;
+	TRange<float> StartTimeInteractionXRange = TRange<float>::Inclusive(0.f, 0.f);
+	TRange<float> EndTimeInteractionXRange = TRange<float>::Inclusive(0.f, 0.f);
+	TRange<float> FadeInInteractionXRange = TRange<float>::Inclusive(0.f, 0.f);
+	TRange<float> FadeOutInteractionXRange = TRange<float>::Inclusive(0.f, 0.f);
 
-	
 	const float InteractionPixelXDelta = 10;
 	const float InteractionRatioYDelta = 0.07f;
 	const float MouseWheelStep = 0.25f;
+
+	double PixelsPerFrame = 0.0;
+
+	enum class ETrimFadeInteractionType : uint8
+	{
+		None = 0,
+		ScrubbingLeftHandle,
+		ScrubbingRightHandle,
+		ScrubbingFadeIn, 
+		ScrubbingFadeOut
+	} TrimFadeInteractionType = ETrimFadeInteractionType::None;
+
+	ETrimFadeInteractionType GetInteractionTypeFromCursorPosition(const FVector2D& InLocalCursorPosition, const FGeometry& WidgetGeometry) const;
 };
\ No newline at end of file
