[WaveformEditor] UI interactions now update the USoundWave asset on ValueSet #robomerge 5.1 #trivial # rb none #preflight 630f398ae54ec9d581aa0137

diff --git a/Engine/Plugins/Editor/WaveformEditor/Source/WaveformEditorWidgets/Public/IWaveformTransformationRenderer.h b/Engine/Plugins/Editor/WaveformEditor/Source/WaveformEditorWidgets/Public/IWaveformTransformationRenderer.h
--- a/Engine/Plugins/Editor/WaveformEditor/Source/WaveformEditorWidgets/Public/IWaveformTransformationRenderer.h
+++ b/Engine/Plugins/Editor/WaveformEditor/Source/WaveformEditorWidgets/Public/IWaveformTransformationRenderer.h
@@ -63,7 +63,7 @@ protected:
 		FProperty* Property = EditedTransformation->GetClass()->FindPropertyByName(PropertyName);
 		check(Property);
 
-		const USoundWave* ParentSoundWave = EditedTransformation->GetTypedOuter<USoundWave>();
+		USoundWave* ParentSoundWave = EditedTransformation->GetTypedOuter<USoundWave>();
 		check(ParentSoundWave)
 
 		FProperty* TransformationsProperty = ParentSoundWave->GetClass()->FindPropertyByName(GET_MEMBER_NAME_CHECKED(USoundWave, Transformations));
@@ -75,8 +75,10 @@ protected:
 		EditChain.SetActivePropertyNode(Property);
 
 		FPropertyChangedEvent PropertyEvent(Property, InChangeType);
+		FPropertyChangedEvent TransformationsPropertyEvent(TransformationsProperty, InChangeType);
 
 		EditedTransformation->PostEditChangeProperty(PropertyEvent);
+		ParentSoundWave->PostEditChangeProperty(TransformationsPropertyEvent);
 
 		TransformationChangeNotifier(PropertyEvent, &EditChain);
 	}
diff --git a/Engine/Plugins/Editor/WaveformEditor/Source/WaveformTransformationsWidgets/Private/WaveformTransformationTrimFadeRenderer.cpp b/Engine/Plugins/Editor/WaveformEditor/Source/WaveformTransformationsWidgets/Private/WaveformTransformationTrimFadeRenderer.cpp
--- a/Engine/Plugins/Editor/WaveformEditor/Source/WaveformTransformationsWidgets/Private/WaveformTransformationTrimFadeRenderer.cpp
+++ b/Engine/Plugins/Editor/WaveformEditor/Source/WaveformTransformationsWidgets/Private/WaveformTransformationTrimFadeRenderer.cpp
@@ -183,7 +183,7 @@ float FWaveformTransformationTrimFadeRenderer::ConvertXRatioToTime(const float I
 {
 	const float NumFrames = TransformationWaveInfo.NumAvilableSamples / TransformationWaveInfo.NumChannels;
 	const float FrameSelected = NumFrames * InRatio;
-	return FrameSelected * TransformationWaveInfo.NumChannels / TransformationWaveInfo.SampleRate;
+	return FrameSelected / TransformationWaveInfo.SampleRate;
 }
 
 void FWaveformTransformationTrimFadeRenderer::Tick(const FGeometry& AllottedGeometry, const double InCurrentTime, const float InDeltaTime)
@@ -218,23 +218,12 @@ void FWaveformTransformationTrimFadeRenderer::UpdateInteractionRange()
 	FadeOutInteractionXRange.SetUpperBoundValue(FadeOutStartX + InteractionPixelXDelta);
 }
 
-FReply FWaveformTransformationTrimFadeRenderer::OnMouseButtonUp(SWidget& OwnerWidget, const FGeometry& MyGeometry, const FPointerEvent& MouseEvent)
-{
-	if (TrimFadeInteractionType != ETrimFadeInteractionType::None)
-	{
-		TrimFadeInteractionType = ETrimFadeInteractionType::None;
-		return FReply::Handled().ReleaseMouseCapture();
-	}
-
-	return FReply::Unhandled();
-}
-
 FReply FWaveformTransformationTrimFadeRenderer::OnMouseButtonDown(SWidget& OwnerWidget, const FGeometry& MyGeometry, const FPointerEvent& MouseEvent)
 {
 	const FVector2D LocalCursorPosition = GetLocalCursorPosition(MouseEvent, MyGeometry);
 
 	const bool bIsLeftMouseButton = MouseEvent.IsMouseButtonDown(EKeys::LeftMouseButton);
-	
+
 	if (bIsLeftMouseButton)
 	{
 		TrimFadeInteractionType = GetInteractionTypeFromCursorPosition(LocalCursorPosition, MyGeometry);
@@ -253,33 +242,7 @@ FReply FWaveformTransformationTrimFadeRenderer::OnMouseMove(SWidget& OwnerWidget
 {
 	if (MouseEvent.IsMouseButtonDown(EKeys::LeftMouseButton) && TrimFadeInteractionType != ETrimFadeInteractionType::None)
 	{
-		const FVector2D LocalCursorPosition = GetLocalCursorPosition(MouseEvent, MyGeometry);
-		const float LocalCursorXRatio = FMath::Clamp(LocalCursorPosition.X / MyGeometry.GetLocalSize().X, 0.f, 1.f);
-		const float SelectedTime = ConvertXRatioToTime(LocalCursorXRatio);
-
-		switch (TrimFadeInteractionType)
-		{
-		case FWaveformTransformationTrimFadeRenderer::ETrimFadeInteractionType::None:
-			break;
-		case FWaveformTransformationTrimFadeRenderer::ETrimFadeInteractionType::ScrubbingLeftHandle:
-			TrimFadeTransform->StartTime = SelectedTime;
-			NotifyTransformationPropertyChanged(TrimFadeTransform, GET_MEMBER_NAME_CHECKED(UWaveformTransformationTrimFade, StartTime), EPropertyChangeType::ValueSet);
-			break;
-		case FWaveformTransformationTrimFadeRenderer::ETrimFadeInteractionType::ScrubbingRightHandle:
-			TrimFadeTransform->EndTime = SelectedTime;
-			NotifyTransformationPropertyChanged(TrimFadeTransform, GET_MEMBER_NAME_CHECKED(UWaveformTransformationTrimFade, EndTime), EPropertyChangeType::ValueSet);
-			break;
-		case FWaveformTransformationTrimFadeRenderer::ETrimFadeInteractionType::ScrubbingFadeIn:
-			TrimFadeTransform->StartFadeTime = FMath::Clamp(SelectedTime - TrimFadeTransform->StartTime, 0.f, TNumericLimits<float>().Max());
-			NotifyTransformationPropertyChanged(TrimFadeTransform, GET_MEMBER_NAME_CHECKED(UWaveformTransformationTrimFade, StartFadeTime), EPropertyChangeType::ValueSet);
-			break;
-		case FWaveformTransformationTrimFadeRenderer::ETrimFadeInteractionType::ScrubbingFadeOut:
-			TrimFadeTransform->EndFadeTime = FMath::Clamp(TrimFadeTransform->EndTime - SelectedTime, 0.f, TNumericLimits<float>().Max());
-			NotifyTransformationPropertyChanged(TrimFadeTransform, GET_MEMBER_NAME_CHECKED(UWaveformTransformationTrimFade, EndFadeTime), EPropertyChangeType::ValueSet);
-			break;
-		default:
-			break;
-		}
+		SetPropertyValueDependingOnInteractionType(MouseEvent, MyGeometry, EPropertyChangeType::Interactive);
 
 		return FReply::Handled().CaptureMouse(OwnerWidget.AsShared());
 	}
@@ -287,6 +250,51 @@ FReply FWaveformTransformationTrimFadeRenderer::OnMouseMove(SWidget& OwnerWidget
 	return FReply::Unhandled();
 }
 
+FReply FWaveformTransformationTrimFadeRenderer::OnMouseButtonUp(SWidget& OwnerWidget, const FGeometry& MyGeometry, const FPointerEvent& MouseEvent)
+{
+	if (TrimFadeInteractionType != ETrimFadeInteractionType::None)
+	{
+		SetPropertyValueDependingOnInteractionType(MouseEvent, MyGeometry, EPropertyChangeType::ValueSet);
+
+		TrimFadeInteractionType = ETrimFadeInteractionType::None;
+		return FReply::Handled().ReleaseMouseCapture();
+	}
+
+	return FReply::Unhandled();
+}
+
+void FWaveformTransformationTrimFadeRenderer::SetPropertyValueDependingOnInteractionType(const FPointerEvent& MouseEvent, const FGeometry& WidgetGeometry, const EPropertyChangeType::Type DesiredChangeType)
+{
+	const FVector2D LocalCursorPosition = GetLocalCursorPosition(MouseEvent, WidgetGeometry);
+	const float LocalCursorXRatio = FMath::Clamp(LocalCursorPosition.X / WidgetGeometry.GetLocalSize().X, 0.f, 1.f);
+	const float SelectedTime = ConvertXRatioToTime(LocalCursorXRatio);
+
+	switch (TrimFadeInteractionType)
+	{
+	case FWaveformTransformationTrimFadeRenderer::ETrimFadeInteractionType::None:
+		break;
+	case FWaveformTransformationTrimFadeRenderer::ETrimFadeInteractionType::ScrubbingLeftHandle:
+		TrimFadeTransform->StartTime = SelectedTime;
+		NotifyTransformationPropertyChanged(TrimFadeTransform, GET_MEMBER_NAME_CHECKED(UWaveformTransformationTrimFade, StartTime), DesiredChangeType);
+		break;
+	case FWaveformTransformationTrimFadeRenderer::ETrimFadeInteractionType::ScrubbingRightHandle:
+		TrimFadeTransform->EndTime = SelectedTime;
+		NotifyTransformationPropertyChanged(TrimFadeTransform, GET_MEMBER_NAME_CHECKED(UWaveformTransformationTrimFade, EndTime), DesiredChangeType);
+		break;
+	case FWaveformTransformationTrimFadeRenderer::ETrimFadeInteractionType::ScrubbingFadeIn:
+		TrimFadeTransform->StartFadeTime = FMath::Clamp(SelectedTime - TrimFadeTransform->StartTime, 0.f, TNumericLimits<float>().Max());
+		NotifyTransformationPropertyChanged(TrimFadeTransform, GET_MEMBER_NAME_CHECKED(UWaveformTransformationTrimFade, StartFadeTime), DesiredChangeType);
+		break;
+	case FWaveformTransformationTrimFadeRenderer::ETrimFadeInteractionType::ScrubbingFadeOut:
+		TrimFadeTransform->EndFadeTime = FMath::Clamp(TrimFadeTransform->EndTime - SelectedTime, 0.f, TNumericLimits<float>().Max());
+		NotifyTransformationPropertyChanged(TrimFadeTransform, GET_MEMBER_NAME_CHECKED(UWaveformTransformationTrimFade, EndFadeTime), DesiredChangeType);
+		break;
+	default:
+		break;
+	}
+}
+
+
 FWaveformTransformationTrimFadeRenderer::ETrimFadeInteractionType FWaveformTransformationTrimFadeRenderer::GetInteractionTypeFromCursorPosition(const FVector2D& InLocalCursorPosition, const FGeometry& WidgetGeometry) const
 {
 	if (IsCursorInFadeInInteractionRange(InLocalCursorPosition, WidgetGeometry))
diff --git a/Engine/Plugins/Editor/WaveformEditor/Source/WaveformTransformationsWidgets/Private/WaveformTransformationTrimFadeRenderer.h b/Engine/Plugins/Editor/WaveformEditor/Source/WaveformTransformationsWidgets/Private/WaveformTransformationTrimFadeRenderer.h
--- a/Engine/Plugins/Editor/WaveformEditor/Source/WaveformTransformationsWidgets/Private/WaveformTransformationTrimFadeRenderer.h
+++ b/Engine/Plugins/Editor/WaveformEditor/Source/WaveformTransformationsWidgets/Private/WaveformTransformationTrimFadeRenderer.h
@@ -62,4 +62,5 @@ private:
 	} TrimFadeInteractionType = ETrimFadeInteractionType::None;
 
 	ETrimFadeInteractionType GetInteractionTypeFromCursorPosition(const FVector2D& InLocalCursorPosition, const FGeometry& WidgetGeometry) const;
+	void SetPropertyValueDependingOnInteractionType(const FPointerEvent& MouseEvent, const FGeometry& WidgetGeometry, const EPropertyChangeType::Type DesiredChangeType);
 };
\ No newline at end of file
