[WaveformTransformationTrimFade] Init EndTime to be SoundWave duration and make it so it is always > StartTime
Fix for:

diff --git a/Engine/Plugins/Editor/WaveformEditor/Source/WaveformTransformations/Private/WaveformTransformationTrimFade.cpp b/Engine/Plugins/Editor/WaveformEditor/Source/WaveformTransformations/Private/WaveformTransformationTrimFade.cpp
--- a/Engine/Plugins/Editor/WaveformEditor/Source/WaveformTransformations/Private/WaveformTransformationTrimFade.cpp
+++ b/Engine/Plugins/Editor/WaveformEditor/Source/WaveformTransformations/Private/WaveformTransformationTrimFade.cpp
@@ -1,6 +1,13 @@
 // Copyright Epic Games, Inc. All Rights Reserved.
 
 #include "WaveformTransformationTrimFade.h"
+#include "Sound/SoundWave.h"
+
+namespace WaveformTransformationTrimFadeNames
+{
+	static FLazyName StartTimeName(GET_MEMBER_NAME_CHECKED(UWaveformTransformationTrimFade, StartTime));
+	static FLazyName EndTimeName(GET_MEMBER_NAME_CHECKED(UWaveformTransformationTrimFade, EndTime));
+}
 
 static void ApplyFadeIn(Audio::FAlignedFloatBuffer& InputAudio, const float FadeLength, const float FadeCurve, const int32 NumChannels, const float SampleRate)
 {
@@ -127,3 +134,57 @@ Audio::FTransformationPtr UWaveformTransformationTrimFade::CreateTransformation(
 {
 	return MakeUnique<FWaveTransformationTrimFade>(StartTime, EndTime, StartFadeTime, StartFadeCurve, EndFadeTime, EndFadeCurve);
 }
+
+void UWaveformTransformationTrimFade::PostEditChangeProperty(struct FPropertyChangedEvent& PropertyChangedEvent)
+{
+	Super::PostEditChangeProperty(PropertyChangedEvent);
+	
+	if (PropertyChangedEvent.GetPropertyName() == WaveformTransformationTrimFadeNames::StartTimeName 
+		|| PropertyChangedEvent.GetPropertyName() == WaveformTransformationTrimFadeNames::EndTimeName)
+	{
+		if (AvailableWaveformDuration < 0.f)
+		{
+			UpdateAvailableWaveformDuration();
+		}
+
+		StartTime = FMath::Clamp(StartTime, 0.f, AvailableWaveformDuration - UE_KINDA_SMALL_NUMBER);
+		EndTime = FMath::Clamp(EndTime, StartTime + UE_KINDA_SMALL_NUMBER, AvailableWaveformDuration);
+	}
+}
+
+void UWaveformTransformationTrimFade::PostInitProperties()
+{
+	Super::PostInitProperties();
+	
+	if (EndTime < 0.f)
+	{
+		UpdateAvailableWaveformDuration();
+		EndTime = AvailableWaveformDuration;
+	}
+}
+
+void UWaveformTransformationTrimFade::UpdateAvailableWaveformDuration()
+{
+	const USoundWave* ParentSoundWave = this->GetTypedOuter<USoundWave>();
+
+	if (!ParentSoundWave)
+	{
+		return;
+	}
+
+	AvailableWaveformDuration = ParentSoundWave->Duration;
+	
+	for (TObjectPtr<UWaveformTransformationBase> Transformation : ParentSoundWave->Transformations)
+	{
+		if (Transformation && Transformation->IsA(UWaveformTransformationTrimFade::StaticClass()))
+		{
+			if (Transformation == this)
+			{
+				return;
+			}
+
+			TObjectPtr<UWaveformTransformationTrimFade> TrimFadeTransformPtr = Cast<UWaveformTransformationTrimFade>(Transformation);
+			AvailableWaveformDuration = TrimFadeTransformPtr->EndTime - TrimFadeTransformPtr->StartTime;
+		}
+	}
+}
diff --git a/Engine/Plugins/Editor/WaveformEditor/Source/WaveformTransformations/Public/WaveformTransformationTrimFade.h b/Engine/Plugins/Editor/WaveformEditor/Source/WaveformTransformations/Public/WaveformTransformationTrimFade.h
--- a/Engine/Plugins/Editor/WaveformEditor/Source/WaveformTransformations/Public/WaveformTransformationTrimFade.h
+++ b/Engine/Plugins/Editor/WaveformEditor/Source/WaveformTransformations/Public/WaveformTransformationTrimFade.h
@@ -30,7 +30,6 @@ UCLASS()
 class WAVEFORMTRANSFORMATIONS_API UWaveformTransformationTrimFade : public UWaveformTransformationBase
 {
 	GENERATED_BODY()
-
 public:
 
 	UPROPERTY(EditAnywhere, Category = "Trim", meta=(ClampMin = 0.0))
@@ -51,9 +50,17 @@ public:
 	UPROPERTY(EditAnywhere, Category = "Fade", meta=(ClampMin = 0.0, ClampMax = 10.0, DisplayName = "Fade-Out Curve"))
 	float EndFadeCurve = 1.f;
 
-	
-public:
 	virtual Audio::FTransformationPtr CreateTransformation() const override;
+
+	void PostEditChangeProperty(struct FPropertyChangedEvent& PropertyChangedEvent) override;
+	void PostInitProperties() override;
+
+	
+
+private: 
+	void UpdateAvailableWaveformDuration();
+	float AvailableWaveformDuration = - 1.f;
+
 };
 
 // class UBaseClass;
