[WaveformTransformTrimFade] Make EndTime trim always point at last sample in frame
#trivial #rb none

diff --git a/Engine/Plugins/Editor/WaveformEditor/Source/WaveformTransformations/Private/WaveformTransformationTrimFade.cpp b/Engine/Plugins/Editor/WaveformEditor/Source/WaveformTransformations/Private/WaveformTransformationTrimFade.cpp
--- a/Engine/Plugins/Editor/WaveformEditor/Source/WaveformTransformations/Private/WaveformTransformationTrimFade.cpp
+++ b/Engine/Plugins/Editor/WaveformEditor/Source/WaveformTransformations/Private/WaveformTransformationTrimFade.cpp
@@ -82,16 +82,16 @@ void FWaveTransformationTrimFade::ProcessAudio(Audio::FWaveformTransformationWav
 	}
 
 	const int32 StartSample = FMath::RoundToInt32(StartTime * InOutWaveInfo.SampleRate) * InOutWaveInfo.NumChannels;
-	int32 EndSample = InputAudio.Num();
-
+	int32 EndSample = InputAudio.Num() - 1;
+	
 	if(EndTime > 0.f)
-	{
-		const int32 IndexOfLastSampleInFrame = InOutWaveInfo.NumChannels - 1;
-		EndSample = FMath::RoundToInt32(EndTime * InOutWaveInfo.SampleRate) * InOutWaveInfo.NumChannels + IndexOfLastSampleInFrame;
-		EndSample = FMath::Min(EndSample, InputAudio.Num());
+	{	
+		const int32 EndFrame = FMath::RoundToInt32(EndTime * InOutWaveInfo.SampleRate);
+		EndSample = EndFrame * InOutWaveInfo.NumChannels - 1;
+		EndSample = FMath::Min(EndSample, InputAudio.Num() - 1);
 	}
 
-	const int32 FinalSize = EndSample - StartSample;
+	const int32 FinalSize = EndSample - StartSample + 1;
 
 	InOutWaveInfo.StartFrameOffset = StartSample - (StartSample % InOutWaveInfo.NumChannels);
 	InOutWaveInfo.NumEditedSamples = FinalSize;
