[WaveformTransformations] Hide Effect Chain Normalization from Content Browser #rb none # trivial #preflight 63184f74ab25bb91f513f4d3 #robomerge 5.1

diff --git a/Engine/Plugins/Editor/WaveformEditor/Source/WaveformTransformations/Public/WaveformTransformationEffectChain.h b/Engine/Plugins/Editor/WaveformEditor/Source/WaveformTransformations/Public/WaveformTransformationEffectChain.h
--- a/Engine/Plugins/Editor/WaveformEditor/Source/WaveformTransformations/Public/WaveformTransformationEffectChain.h
+++ b/Engine/Plugins/Editor/WaveformEditor/Source/WaveformTransformations/Public/WaveformTransformationEffectChain.h
@@ -22,7 +22,7 @@ private:
 	TArray<TStrongObjectPtr<class USoundEffectSourcePreset>> Presets;
 };
 
-UCLASS()
+UCLASS(hidden)
 class WAVEFORMTRANSFORMATIONS_API UWaveformTransformationEffectChain : public UWaveformTransformationBase
 {
 	GENERATED_BODY()
