[WaveformEditor] Logarithmic Zoom
#trivial

diff --git a/Engine/Plugins/Editor/WaveformEditor/Source/WaveformEditorWidgets/Private/WaveformEditorTransportCoordinator.cpp b/Engine/Plugins/Editor/WaveformEditor/Source/WaveformEditorWidgets/Private/WaveformEditorTransportCoordinator.cpp
--- a/Engine/Plugins/Editor/WaveformEditor/Source/WaveformEditorWidgets/Private/WaveformEditorTransportCoordinator.cpp
+++ b/Engine/Plugins/Editor/WaveformEditor/Source/WaveformEditorWidgets/Private/WaveformEditorTransportCoordinator.cpp
@@ -208,10 +208,9 @@ void FWaveformEditorTransportCoordinator::ReceivePlayBackRatio(const float NewRa
 	}
 }
 
-void FWaveformEditorTransportCoordinator::OnZoomLevelChanged(const uint8 NewLevel)
+void FWaveformEditorTransportCoordinator::OnZoomLevelChanged(const float NewLevel)
 {
-	const float NewZoomRatio = NewLevel / 100.f;
-	UpdateZoomRatioAndDisplayRange(NewZoomRatio);
+	UpdateZoomRatioAndDisplayRange(NewLevel);
 }
 
 float FWaveformEditorTransportCoordinator::ConvertAbsoluteRatioToZoomed(const float InAbsoluteRatio) const
diff --git a/Engine/Plugins/Editor/WaveformEditor/Source/WaveformEditorWidgets/Private/WaveformEditorZoomController.cpp b/Engine/Plugins/Editor/WaveformEditor/Source/WaveformEditorWidgets/Private/WaveformEditorZoomController.cpp
--- a/Engine/Plugins/Editor/WaveformEditor/Source/WaveformEditorWidgets/Private/WaveformEditorZoomController.cpp
+++ b/Engine/Plugins/Editor/WaveformEditor/Source/WaveformEditorWidgets/Private/WaveformEditorZoomController.cpp
@@ -6,7 +6,7 @@ void FWaveformEditorZoomController::ZoomIn()
 {
 	if (CanZoomIn())
 	{
-		ZoomPercentage += ZoomPercentageStep;
+		ZoomLevel = FMath::Clamp(ZoomLevel += ZoomLevelStep, 0, LogRatioBase);
 		ApplyZoom();
 	}
 
@@ -14,21 +14,21 @@ void FWaveformEditorZoomController::ZoomIn()
 
 bool FWaveformEditorZoomController::CanZoomIn() const
 {
-	return ZoomPercentage + ZoomPercentageStep <= 100;
+	return ZoomLevel + ZoomLevelStep <= LogRatioBase + ZoomLevelInitValue;
 }
 
 void FWaveformEditorZoomController::ZoomOut()
 {
 	if (CanZoomOut())
 	{
-		ZoomPercentage -= ZoomPercentageStep;
+		ZoomLevel = FMath::Clamp(ZoomLevel -= ZoomLevelStep, 0, LogRatioBase);
 		ApplyZoom();
 	}
 }
 
 bool FWaveformEditorZoomController::CanZoomOut() const
 {
-	return ZoomPercentage - ZoomPercentageStep >= 0;
+	return ZoomLevel - ZoomLevelStep >= 0;
 }
 
 
@@ -44,13 +44,17 @@ void FWaveformEditorZoomController::ZoomByDelta(const float Delta)
 	}
 }
 
-uint8 FWaveformEditorZoomController::GetZoomRatio() const
+float FWaveformEditorZoomController::GetZoomRatio() const
 {
-	return 100 - ZoomPercentage;
+	return 1 - ConvertZoomLevelToLogRatio();
 }
 
 void FWaveformEditorZoomController::ApplyZoom()
 {
-	UE_LOG(LogInit, Log, TEXT("Wave Editor Zoom %d%%"), ZoomPercentage);
-	OnZoomRatioChanged.Broadcast(100 - ZoomPercentage);
+	OnZoomRatioChanged.Broadcast(1 - ConvertZoomLevelToLogRatio());
+}
+
+float FWaveformEditorZoomController::ConvertZoomLevelToLogRatio() const
+{
+	return FMath::Clamp(FMath::LogX(LogRatioBase, ZoomLevel), 0.f, 1.f);
 }
diff --git a/Engine/Plugins/Editor/WaveformEditor/Source/WaveformEditorWidgets/Public/WaveformEditorTransportCoordinator.h b/Engine/Plugins/Editor/WaveformEditor/Source/WaveformEditorWidgets/Public/WaveformEditorTransportCoordinator.h
--- a/Engine/Plugins/Editor/WaveformEditor/Source/WaveformEditorWidgets/Public/WaveformEditorTransportCoordinator.h
+++ b/Engine/Plugins/Editor/WaveformEditor/Source/WaveformEditorWidgets/Public/WaveformEditorTransportCoordinator.h
@@ -40,7 +40,7 @@ public:
 
 	const float GetPlayheadPosition() const;
 	void ReceivePlayBackRatio(const float NewRatio);
-	void OnZoomLevelChanged(const uint8 NewLevel);
+	void OnZoomLevelChanged(const float NewLevel);
 
 	float ConvertAbsoluteRatioToZoomed(const float InAbsoluteRatio) const;
 	float ConvertZoomedRatioToAbsolute(const float InZoomedRatio) const;
diff --git a/Engine/Plugins/Editor/WaveformEditor/Source/WaveformEditorWidgets/Public/WaveformEditorZoomController.h b/Engine/Plugins/Editor/WaveformEditor/Source/WaveformEditorWidgets/Public/WaveformEditorZoomController.h
--- a/Engine/Plugins/Editor/WaveformEditor/Source/WaveformEditorWidgets/Public/WaveformEditorZoomController.h
+++ b/Engine/Plugins/Editor/WaveformEditor/Source/WaveformEditorWidgets/Public/WaveformEditorZoomController.h
@@ -4,7 +4,7 @@
 
 #include "Delegates/Delegate.h"
 
-DECLARE_MULTICAST_DELEGATE_OneParam(FOnZoomRatioChanged, const uint8 /* New Zoom Ratio */);
+DECLARE_MULTICAST_DELEGATE_OneParam(FOnZoomRatioChanged, const float /* New Zoom Ratio */);
 
 class WAVEFORMEDITORWIDGETS_API FWaveformEditorZoomController
 {
@@ -17,13 +17,16 @@ public:
 	bool CanZoomOut() const;
 	void ZoomByDelta(const float Delta);
 
-	uint8 GetZoomRatio() const;
+	float GetZoomRatio() const;
 
 	FOnZoomRatioChanged OnZoomRatioChanged;
 
 private:
 	void ApplyZoom();
-
-	uint8 ZoomPercentage = 0;
-	uint8 ZoomPercentageStep = 5;
+	float ConvertZoomLevelToLogRatio() const;
+	
+	float ZoomLevelInitValue = 1.f;
+	float ZoomLevel = ZoomLevelInitValue;
+	float ZoomLevelStep = 2.f;
+	const uint32 LogRatioBase = 100;
 };
\ No newline at end of file
