[WaveformEditor] Allow user to select between different channels configurations when bouncing

diff --git a/Engine/Plugins/Editor/WaveformEditor/Source/WaveformEditor/Private/WaveformEditor.cpp b/Engine/Plugins/Editor/WaveformEditor/Source/WaveformEditor/Private/WaveformEditor.cpp
--- a/Engine/Plugins/Editor/WaveformEditor/Source/WaveformEditor/Private/WaveformEditor.cpp
+++ b/Engine/Plugins/Editor/WaveformEditor/Source/WaveformEditor/Private/WaveformEditor.cpp
@@ -14,6 +14,7 @@
 #include "WaveformEditorLog.h"
 #include "WaveformEditorRenderData.h"
 #include "WaveformEditorStyle.h"
+#include "WaveformEditorToolMenuContext.h"
 #include "WaveformEditorWaveWriter.h"
 #include "WaveformTransformationsRenderManager.h"
 #include "Widgets/Docking/SDockTab.h"
@@ -243,15 +244,50 @@ bool FWaveformEditor::RegisterToolbar()
 			LOCTEXT("WaveformEditorRender", ""),
 			LOCTEXT("WaveformEditorRenderButtonTooltip", "Exports the edited waveform to a USoundWave asset"),
 			FSlateIcon(FAppStyle::GetAppStyleSetName(), "GenericCurveEditor.SetViewModeAbsolute")
-		);
+		);	
 
 		ExportSection.AddEntry(ExportEntry);
 
+		ExportSection.AddDynamicEntry("ExportOptions", FNewToolMenuSectionDelegate::CreateLambda([this](FToolMenuSection& InSection) {
+				UWaveformEditorToolMenuContext* Context = InSection.FindContext<UWaveformEditorToolMenuContext>();
+				
+				if (Context && Context->WaveformEditor.IsValid())
+				{
+					FToolMenuEntry ExportOptionsEntry = FToolMenuEntry::InitComboButton(
+						"ExportsOptionsCombo",
+						FToolUIActionChoice(FUIAction()),
+						FNewToolMenuChoice(FOnGetContent::CreateSP(Context->WaveformEditor.Pin().Get(), &FWaveformEditor::GenerateExportOptionsMenu)),
+						LOCTEXT("ExportsOptions_Label", "Export Options"),
+						LOCTEXT("ExportsOptions_ToolTip", "Export options for this waveform"),
+						FSlateIcon(FAppStyle::GetAppStyleSetName(), "GenericCurveEditor.SetViewModeAbsolute"),
+						true
+					);
+
+					InSection.AddEntry(ExportOptionsEntry);
+				}				
+			}
+		));
+
 	}
 
 	return true;
 }
 
+TSharedRef<SWidget> FWaveformEditor::GenerateExportOptionsMenu()
+{
+	const bool bShouldCloseWindowAfterMenuSelection = true;
+	FMenuBuilder MenuBuilder(bShouldCloseWindowAfterMenuSelection, GetToolkitCommands());
+	
+	MenuBuilder.BeginSection(NAME_None, LOCTEXT("ChannelSection_Label", "Export Channel Format"));
+	{
+		MenuBuilder.AddMenuEntry(FWaveformEditorCommands::Get().ExportFormatMono);
+		MenuBuilder.AddMenuEntry(FWaveformEditorCommands::Get().ExportFormatStereo);
+	}
+	MenuBuilder.EndSection();
+
+	return MenuBuilder.MakeWidget();
+}
+
 bool FWaveformEditor::BindCommands()
 {
 	const FWaveformEditorCommands& Commands = FWaveformEditorCommands::Get();
@@ -290,6 +326,18 @@ bool FWaveformEditor::BindCommands()
 		FExecuteAction::CreateSP(this, &FWaveformEditor::ExportWaveform),
 		FCanExecuteAction::CreateSP(WaveWriter.ToSharedRef(), &FWaveformEditorWaveWriter::CanCreateSoundWaveAsset));
 
+	ToolkitCommands->MapAction(
+		Commands.ExportFormatMono,
+		FExecuteAction::CreateLambda([this] { (WaveWriter->SetExportChannelsFormat(WaveformEditorWaveWriter::EChannelFormat::Mono)); }),
+		FCanExecuteAction::CreateLambda([this] {return WaveWriter.IsValid(); }),
+		FIsActionChecked::CreateLambda([this] {return WaveWriter->GetExportChannelsFormat() == WaveformEditorWaveWriter::EChannelFormat::Mono; }));
+
+	ToolkitCommands->MapAction(
+		Commands.ExportFormatStereo,
+		FExecuteAction::CreateLambda([this] { (WaveWriter->SetExportChannelsFormat(WaveformEditorWaveWriter::EChannelFormat::Stereo)); }),
+		FCanExecuteAction::CreateLambda([this] {return WaveWriter.IsValid(); }),
+		FIsActionChecked::CreateLambda([this] {return WaveWriter->GetExportChannelsFormat() == WaveformEditorWaveWriter::EChannelFormat::Stereo; }));
+
 	return true;
 }
 
@@ -381,6 +429,15 @@ bool FWaveformEditor::MatchesContext(const FTransactionContext& InContext, const
 	return bShoouldUndo;
 }
 
+void FWaveformEditor::InitToolMenuContext(FToolMenuContext& MenuContext)
+{
+	FAssetEditorToolkit::InitToolMenuContext(MenuContext);
+
+	UWaveformEditorToolMenuContext* Context = NewObject<UWaveformEditorToolMenuContext>();
+	Context->WaveformEditor = SharedThis(this);
+	MenuContext.AddObject(Context);
+}
+
 bool FWaveformEditor::SetUpDetailsViews()
 {
 	if (SoundWave == nullptr)
diff --git a/Engine/Plugins/Editor/WaveformEditor/Source/WaveformEditor/Private/WaveformEditorCommands.cpp b/Engine/Plugins/Editor/WaveformEditor/Source/WaveformEditor/Private/WaveformEditorCommands.cpp
--- a/Engine/Plugins/Editor/WaveformEditor/Source/WaveformEditor/Private/WaveformEditorCommands.cpp
+++ b/Engine/Plugins/Editor/WaveformEditor/Source/WaveformEditor/Private/WaveformEditorCommands.cpp
@@ -22,7 +22,10 @@ void FWaveformEditorCommands::RegisterCommands()
 
 	UI_COMMAND(ZoomIn, "Zoom In", "Zooms into the sound wave", EUserInterfaceActionType::Button, FInputChord());
 	UI_COMMAND(ZoomOut, "Zoom Out", "Zooms out from the sound wave", EUserInterfaceActionType::Button, FInputChord());
+
 	UI_COMMAND(ExportWaveform, "Export Waveform", "Exports the edited waveform to a new sound wave uasset", EUserInterfaceActionType::Button, FInputChord());
+	UI_COMMAND(ExportFormatMono, "Mono", "Sets the export format to mono", EUserInterfaceActionType::RadioButton, FInputChord());
+	UI_COMMAND(ExportFormatStereo, "Stereo", "Sets the export format to stereo", EUserInterfaceActionType::RadioButton, FInputChord());
 }
 	
 const FWaveformEditorCommands& FWaveformEditorCommands::Get()
diff --git a/Engine/Plugins/Editor/WaveformEditor/Source/WaveformEditor/Private/WaveformEditorWaveWriter.cpp b/Engine/Plugins/Editor/WaveformEditor/Source/WaveformEditor/Private/WaveformEditorWaveWriter.cpp
--- a/Engine/Plugins/Editor/WaveformEditor/Source/WaveformEditor/Private/WaveformEditorWaveWriter.cpp
+++ b/Engine/Plugins/Editor/WaveformEditor/Source/WaveformEditor/Private/WaveformEditorWaveWriter.cpp
@@ -11,6 +11,7 @@
 FWaveformEditorWaveWriter::FWaveformEditorWaveWriter(USoundWave* InSoundWave)
 	: SourceSoundWave(InSoundWave)
 	, WaveWriter(MakeUnique<Audio::FSoundWavePCMWriter>())
+	, TargetChannelsFormat(GetSupportedFormatFromChannelCount(InSoundWave->NumChannels))
 {
 }
 
@@ -56,8 +57,36 @@ void FWaveformEditorWaveWriter::ExportTransformedWaveform()
 	}
 }
 
+WaveformEditorWaveWriter::EChannelFormat FWaveformEditorWaveWriter::GetExportChannelsFormat() const
+{
+	return TargetChannelsFormat;
+}
+
+void FWaveformEditorWaveWriter::SetExportChannelsFormat(const WaveformEditorWaveWriter::EChannelFormat InTargetChannelFormat)
+{
+	TargetChannelsFormat = InTargetChannelFormat;
+}
+
+WaveformEditorWaveWriter::EChannelFormat FWaveformEditorWaveWriter::GetSupportedFormatFromChannelCount(const int InChannelCount) const
+{
+	WaveformEditorWaveWriter::EChannelFormat OutSupportedFormat = DefaultExportFormat;
+	
+	switch (InChannelCount)
+	{
+	case 1:
+		OutSupportedFormat = WaveformEditorWaveWriter::EChannelFormat::Mono;
+		break;
+	case 2:
+		OutSupportedFormat = WaveformEditorWaveWriter::EChannelFormat::Stereo; 
+		break;
+	default:
+		break;
+	}
+
+	return OutSupportedFormat;
+}
 
-Audio::TSampleBuffer<> FWaveformEditorWaveWriter::GenerateSampleBuffer() const
+Audio::TSampleBuffer<int16> FWaveformEditorWaveWriter::GenerateSampleBuffer() const
 {
 	TArray<uint8> RawPCMData;
 	uint16 NumChannels;
@@ -91,12 +120,7 @@ Audio::TSampleBuffer<> FWaveformEditorWaveWriter::GenerateSampleBuffer() const
 			Transformation->ProcessAudio(TransformationInfo);
 		}
 
-		const float MaxValue = Audio::ArrayMaxAbsValue(Buffer);
-
-		if (MaxValue > 1.f)
-		{
-			Audio::ArrayMultiplyByConstantInPlace(Buffer, 1.f / MaxValue);
-		}
+		NormalizeBufferToValue(Buffer, 1.f);
 
 		SampleRate = TransformationInfo.SampleRate;
 		NumChannels = TransformationInfo.NumChannels;
@@ -106,7 +130,82 @@ Audio::TSampleBuffer<> FWaveformEditorWaveWriter::GenerateSampleBuffer() const
 		check(SampleRate > 0);
 	}
 
-	Audio::TSampleBuffer<> OutBuffer(Buffer, NumChannels, SampleRate);
+	if (GetSupportedFormatFromChannelCount(NumChannels) == TargetChannelsFormat)
+	{
+		Audio::TSampleBuffer<> OutBuffer(Buffer, NumChannels, SampleRate);
+		return MoveTemp(OutBuffer);
+	}
 
+		
+	switch (TargetChannelsFormat)
+	{
+	case WaveformEditorWaveWriter::EChannelFormat::Mono:
+		return DownmixBufferToMono(Buffer, NumChannels, SampleRate);
+		break;
+	case WaveformEditorWaveWriter::EChannelFormat::Stereo:
+		return UpmixBufferToStereo(Buffer, NumChannels, SampleRate);
+		break;
+	default:
+		Audio::TSampleBuffer<int16> OutBuffer(Buffer, NumChannels, SampleRate);
+		return MoveTemp(OutBuffer);
+		break;
+	}	
+}
+
+Audio::TSampleBuffer<int16> FWaveformEditorWaveWriter::DownmixBufferToMono(const Audio::FAlignedFloatBuffer& InSampleBuffer, const uint32 InNumChannels, const uint32 InSampleRate) const
+{
+	check(InNumChannels > 0)
+
+	if (InNumChannels < 2)
+	{
+		Audio::TSampleBuffer<int16> OutBuffer(InSampleBuffer, InNumChannels, InSampleRate);
+		return MoveTemp(OutBuffer);
+	}
+
+	Audio::TSampleBuffer<float> FloatSampleBuffer(InSampleBuffer, InNumChannels, InSampleRate);
+	FloatSampleBuffer.MixBufferToChannels(1);
+	Audio::FAlignedFloatBuffer MonoBuffer;
+	MonoBuffer.SetNumUninitialized(FloatSampleBuffer.GetNumSamples());
+	MonoBuffer = FloatSampleBuffer.GetArrayView();
+	NormalizeBufferToValue(MonoBuffer, 1.f);
+
+	Audio::TSampleBuffer<int16> OutMonoBuffer = Audio::TSampleBuffer<int16>(MonoBuffer, 1, FloatSampleBuffer.GetSampleRate());
+
+	return MoveTemp(OutMonoBuffer);
+}
+
+Audio::TSampleBuffer<int16> FWaveformEditorWaveWriter::UpmixBufferToStereo(const Audio::FAlignedFloatBuffer& InSampleBuffer, const uint32 InNumChannels, const uint32 InSampleRate) const
+{
+	check(InNumChannels > 0)
+
+	if (InNumChannels > 1)
+	{
+		Audio::TSampleBuffer<int16> OutBuffer(InSampleBuffer, InNumChannels, InSampleRate);
+		return MoveTemp(OutBuffer);
+	}
+
+	Audio::FAlignedFloatBuffer StereoBuffer;
+	StereoBuffer.SetNumUninitialized(InSampleBuffer.Num() * 2);
+
+	for (int32 i = 0; i < InSampleBuffer.Num(); ++i)
+	{
+		const uint32 StereoFrameIndex = i * 2;
+		StereoBuffer[StereoFrameIndex] = InSampleBuffer[i];
+		StereoBuffer[StereoFrameIndex + 1] = InSampleBuffer[i];
+	}
+
+	NormalizeBufferToValue(StereoBuffer, 1.f);
+
+	Audio::TSampleBuffer<int16> OutBuffer(StereoBuffer, 2, InSampleRate);
 	return MoveTemp(OutBuffer);
+}
+
+void FWaveformEditorWaveWriter::NormalizeBufferToValue(Audio::FAlignedFloatBuffer& InOutBuffer, const float InTargetMaxValue) const
+{
+	const float MaxValue = Audio::ArrayMaxAbsValue(InOutBuffer);
+
+	if (MaxValue > InTargetMaxValue)
+	{
+		Audio::ArrayMultiplyByConstantInPlace(InOutBuffer, InTargetMaxValue / MaxValue);
+	}
 }
\ No newline at end of file
diff --git a/Engine/Plugins/Editor/WaveformEditor/Source/WaveformEditor/Private/WaveformEditorWaveWriter.h b/Engine/Plugins/Editor/WaveformEditor/Source/WaveformEditor/Private/WaveformEditorWaveWriter.h
--- a/Engine/Plugins/Editor/WaveformEditor/Source/WaveformEditor/Private/WaveformEditorWaveWriter.h
+++ b/Engine/Plugins/Editor/WaveformEditor/Source/WaveformEditor/Private/WaveformEditorWaveWriter.h
@@ -6,6 +6,16 @@
 
 class USoundWave;
 
+namespace WaveformEditorWaveWriter
+{
+	enum class EChannelFormat : int32
+	{
+		Mono = 1,
+		Stereo
+	};
+}
+
+
 class FWaveformEditorWaveWriter
 {
 public:
@@ -14,9 +24,20 @@ public:
 	bool CanCreateSoundWaveAsset() const;
 	void ExportTransformedWaveform();
 
+	WaveformEditorWaveWriter::EChannelFormat GetExportChannelsFormat() const;
+	void SetExportChannelsFormat(const WaveformEditorWaveWriter::EChannelFormat InTargetChannelFormat);
+
 private: 
-	Audio::TSampleBuffer<> GenerateSampleBuffer() const;
+	WaveformEditorWaveWriter::EChannelFormat GetSupportedFormatFromChannelCount(const int InChannelCount) const;
+
+	Audio::TSampleBuffer<int16> GenerateSampleBuffer() const;
+	Audio::TSampleBuffer<int16> DownmixBufferToMono(const Audio::FAlignedFloatBuffer& InSampleBuffer, const uint32 InNumChannels, const uint32 InSampleRate) const;
+	Audio::TSampleBuffer<int16> UpmixBufferToStereo(const Audio::FAlignedFloatBuffer& InSampleBuffer, const uint32 InNumChannels, const uint32 InSampleRate) const;
+
+	void NormalizeBufferToValue(Audio::FAlignedFloatBuffer& InOutBuffer, const float InTargetMaxValue) const;
 
+	const WaveformEditorWaveWriter::EChannelFormat DefaultExportFormat = WaveformEditorWaveWriter::EChannelFormat::Stereo;
 	USoundWave* SourceSoundWave = nullptr;
 	TUniquePtr<Audio::FSoundWavePCMWriter> WaveWriter = nullptr;
+	WaveformEditorWaveWriter::EChannelFormat TargetChannelsFormat;
 };
\ No newline at end of file
diff --git a/Engine/Plugins/Editor/WaveformEditor/Source/WaveformEditor/Public/WaveformEditor.h b/Engine/Plugins/Editor/WaveformEditor/Source/WaveformEditor/Public/WaveformEditor.h
--- a/Engine/Plugins/Editor/WaveformEditor/Source/WaveformEditor/Public/WaveformEditor.h
+++ b/Engine/Plugins/Editor/WaveformEditor/Source/WaveformEditor/Public/WaveformEditor.h
@@ -43,6 +43,8 @@ public:
 	void PostRedo(bool bSuccess) override;
 	bool MatchesContext(const FTransactionContext& InContext, const TArray<TPair<UObject*, FTransactionObjectEvent>>& TransactionObjectContexts) const override;
 
+	virtual void InitToolMenuContext(FToolMenuContext& MenuContext) override;
+
 private:
 	bool SetupAudioComponent();
 	bool SetUpTransportController();
@@ -56,6 +58,8 @@ private:
 	/**	Toolbar Setup */
 	bool RegisterToolbar();
 	bool BindCommands();
+	TSharedRef<SWidget> GenerateExportOptionsMenu();
+
 
 	/**	Details tabs set up */
 	TSharedRef<SDockTab> SpawnTab_Properties(const FSpawnTabArgs& Args);
diff --git a/Engine/Plugins/Editor/WaveformEditor/Source/WaveformEditor/Public/WaveformEditorCommands.h b/Engine/Plugins/Editor/WaveformEditor/Source/WaveformEditor/Public/WaveformEditorCommands.h
--- a/Engine/Plugins/Editor/WaveformEditor/Source/WaveformEditor/Public/WaveformEditorCommands.h
+++ b/Engine/Plugins/Editor/WaveformEditor/Source/WaveformEditor/Public/WaveformEditorCommands.h
@@ -24,4 +24,6 @@ public:
 	TSharedPtr<FUICommandInfo> ZoomOut;
 
 	TSharedPtr<FUICommandInfo> ExportWaveform;
+	TSharedPtr<FUICommandInfo> ExportFormatMono;
+	TSharedPtr<FUICommandInfo> ExportFormatStereo;
 };
diff --git a/Engine/Plugins/Editor/WaveformEditor/Source/WaveformEditor/Public/WaveformEditorToolMenuContext.h b/Engine/Plugins/Editor/WaveformEditor/Source/WaveformEditor/Public/WaveformEditorToolMenuContext.h
new file mode 100644
--- /dev/null
+++ b/Engine/Plugins/Editor/WaveformEditor/Source/WaveformEditor/Public/WaveformEditorToolMenuContext.h
@@ -0,0 +1,19 @@
+// Copyright Epic Games, Inc. All Rights Reserved.
+
+#pragma once
+
+#include "CoreMinimal.h"
+
+#include "WaveformEditorToolMenuContext.generated.h"
+
+class FWaveformEditor;
+
+UCLASS()
+class WAVEFORMEDITOR_API UWaveformEditorToolMenuContext : public UObject
+{
+	GENERATED_BODY()
+
+public:
+
+	TWeakPtr<FWaveformEditor> WaveformEditor;
+};
\ No newline at end of file
