[Waveform Editor] Decouple widget style pointers from paint methods

diff --git a/Engine/Plugins/Editor/WaveformEditor/Source/WaveformEditorWidgets/Private/SWaveformEditorTimeRuler.cpp b/Engine/Plugins/Editor/WaveformEditor/Source/WaveformEditorWidgets/Private/SWaveformEditorTimeRuler.cpp
--- a/Engine/Plugins/Editor/WaveformEditor/Source/WaveformEditorWidgets/Private/SWaveformEditorTimeRuler.cpp
+++ b/Engine/Plugins/Editor/WaveformEditor/Source/WaveformEditorWidgets/Private/SWaveformEditorTimeRuler.cpp
@@ -20,15 +20,15 @@ void SWaveformEditorTimeRuler::Construct(const FArguments& InArgs, TSharedRef<FW
 
 	Style = InArgs._Style;
 	
-	check(Style)
-	HandleColor = &Style->HandleColor;
+	check(Style);
+	HandleColor = Style->HandleColor;
 	HandleWidth = Style->HandleWidth;
-	HandleBrush = &Style->HandleBrush;
-	BackgroundColor = &Style->BackgroundColor;
-	BackgroundBrush = &Style->BackgroundBrush;
-	TicksColor = &Style->TicksColor;
-	TicksTextColor = &Style->TicksTextColor;
-	TicksTextFont = &Style->TicksTextFont;
+	HandleBrush = Style->HandleBrush;
+	BackgroundColor = Style->BackgroundColor;
+	BackgroundBrush = Style->BackgroundBrush;
+	TicksColor = Style->TicksColor;
+	TicksTextColor = Style->TicksTextColor;
+	TicksTextFont = Style->TicksTextFont;
 	DesiredWidth = Style->DesiredWidth;
 	DesiredHeight = Style->DesiredHeight;
 	TicksTextOffset = Style->TicksTextOffset;
@@ -45,9 +45,9 @@ int32 SWaveformEditorTimeRuler::OnPaint(const FPaintArgs& Args, const FGeometry&
 		OutDrawElements,
 		++LayerId,
 		AllottedGeometry.ToPaintGeometry(),
-		BackgroundBrush,
+		&BackgroundBrush,
 		ESlateDrawEffect::None,
-		BackgroundColor->GetSpecifiedColor());
+		BackgroundColor.GetSpecifiedColor());
 	
 	DrawRulerTicks(AllottedGeometry, OutDrawElements, LayerId);
 	DrawPlayheadHandle(AllottedGeometry, OutDrawElements, LayerId);
@@ -77,7 +77,7 @@ void SWaveformEditorTimeRuler::DrawRulerTicks(const FGeometry& AllottedGeometry,
 			AllottedGeometry.ToPaintGeometry(),
 			LinePoints,
 			ESlateDrawEffect::None,
-			TicksColor->GetSpecifiedColor(),
+			TicksColor.GetSpecifiedColor(),
 			false);
 
 
@@ -94,7 +94,7 @@ void SWaveformEditorTimeRuler::DrawRulerTicks(const FGeometry& AllottedGeometry,
 				AllottedGeometry.ToPaintGeometry(),
 				LinePoints,
 				ESlateDrawEffect::None,
-				TicksColor->GetSpecifiedColor(),
+				TicksColor.GetSpecifiedColor(),
 				false);
 
 		}
@@ -131,9 +131,9 @@ void SWaveformEditorTimeRuler::DrawTickTimeString(float TickTimeSeconds, const d
 		++LayerId,
 		AllottedGeometry.ToPaintGeometry(TextOffset, AllottedGeometry.Size),
 		TimeString,
-		*TicksTextFont,
+		TicksTextFont,
 		ESlateDrawEffect::None,
-		TicksTextColor->GetSpecifiedColor()
+		TicksTextColor.GetSpecifiedColor()
 	);
 }
 
@@ -212,6 +212,29 @@ void SWaveformEditorTimeRuler::UpdateDisplayUnit(const EWaveformEditorDisplayUni
 	DisplayUnit = InDisplayUnit;
 }
 
+void SWaveformEditorTimeRuler::OnStyleUpdated(const FWaveformEditorWidgetStyleBase* UpdatedStyle)
+{
+	check(UpdatedStyle);
+	check(Style);
+
+	if (UpdatedStyle != Style)
+	{
+		return;
+	}
+
+	HandleColor = Style->HandleColor;
+	HandleWidth = Style->HandleWidth;
+	HandleBrush = Style->HandleBrush;
+	BackgroundColor = Style->BackgroundColor;
+	BackgroundBrush = Style->BackgroundBrush;
+	TicksColor = Style->TicksColor;
+	TicksTextColor = Style->TicksTextColor;
+	TicksTextFont = Style->TicksTextFont;
+	DesiredWidth = Style->DesiredWidth;
+	DesiredHeight = Style->DesiredHeight;
+	TicksTextOffset = Style->TicksTextOffset;
+}
+
 void SWaveformEditorTimeRuler::DrawPlayheadHandle(const FGeometry& AllottedGeometry, FSlateWindowElementList& OutDrawElements, int32& LayerId) const
 {
 	const float WindowWidth = AllottedGeometry.Size.X;
@@ -223,9 +246,9 @@ void SWaveformEditorTimeRuler::DrawPlayheadHandle(const FGeometry& AllottedGeome
 		OutDrawElements,
 		LayerId,
 		HandleGeometry,
-		HandleBrush,
+		&HandleBrush,
 		ESlateDrawEffect::None,
-		HandleColor->GetSpecifiedColor()
+		HandleColor.GetSpecifiedColor()
 	);
 
 	++LayerId;
diff --git a/Engine/Plugins/Editor/WaveformEditor/Source/WaveformEditorWidgets/Private/SWaveformEditorTimeRuler.h b/Engine/Plugins/Editor/WaveformEditor/Source/WaveformEditorWidgets/Private/SWaveformEditorTimeRuler.h
--- a/Engine/Plugins/Editor/WaveformEditor/Source/WaveformEditorWidgets/Private/SWaveformEditorTimeRuler.h
+++ b/Engine/Plugins/Editor/WaveformEditor/Source/WaveformEditorWidgets/Private/SWaveformEditorTimeRuler.h
@@ -21,9 +21,6 @@ class SWaveformEditorTimeRuler : public SCompoundWidget
 	SLATE_BEGIN_ARGS(SWaveformEditorTimeRuler)
 		: _DisplayUnit(EWaveformEditorDisplayUnit::Seconds)
 	{
-		const ISlateStyle* WaveEditorStyle = FSlateStyleRegistry::FindSlateStyle("WaveformEditorStyle");
-		check(WaveEditorStyle)
-		_Style = &WaveEditorStyle->GetWidgetStyle<FWaveformEditorTimeRulerStyle>("WaveformEditorRuler.Style");
 	}
 
 	SLATE_ARGUMENT(EWaveformEditorDisplayUnit, DisplayUnit)
@@ -37,6 +34,8 @@ public:
 	void UpdateGridMetrics(const FWaveEditorGridMetrics& InMetrics);
 	void UpdateDisplayUnit(const EWaveformEditorDisplayUnit InDisplayUnit);
 
+	void OnStyleUpdated(const FWaveformEditorWidgetStyleBase* UpdatedStyle);
+
 	/** Delegate sent when the user selects a new display unit from the RMB menu*/
 	FOnTimeUnitMenuSelection OnTimeUnitMenuSelection;
 
@@ -60,17 +59,17 @@ private:
 
 	const FWaveformEditorTimeRulerStyle* Style = nullptr;
 
-	const FSlateColor* BackgroundColor = nullptr;
-	const FSlateBrush* BackgroundBrush = nullptr;
-	const FSlateBrush* HandleBrush = nullptr;
-	const FSlateColor* HandleColor = nullptr;
-	const FSlateColor* TicksColor = nullptr;
-	const FSlateColor* TicksTextColor = nullptr;
+	FSlateColor BackgroundColor = FLinearColor::Black;
+	FSlateBrush BackgroundBrush;
+	FSlateBrush HandleBrush;
+	FSlateColor HandleColor = FLinearColor(255.f, 0.1f, 0.2f, 1.f);
+	FSlateColor TicksColor = FLinearColor(1.f, 1.f, 1.f, 0.9f);
+	FSlateColor TicksTextColor = FLinearColor(1.f, 1.f, 1.f, 0.9f);
 	float DesiredHeight = 0.f;
 	float DesiredWidth = 0.f;
-	float HandleWidth = 0.f;
+	float HandleWidth = 15.f;
 	float TicksTextOffset = 5.f;
-	const FSlateFontInfo* TicksTextFont;
+	FSlateFontInfo TicksTextFont;
 
 	TSharedPtr<FWaveformEditorTransportCoordinator> TransportCoordinator = nullptr;
 	TSharedPtr<FWaveformEditorRenderData> RenderData = nullptr;
diff --git a/Engine/Plugins/Editor/WaveformEditor/Source/WaveformEditorWidgets/Private/SWaveformPanel.cpp b/Engine/Plugins/Editor/WaveformEditor/Source/WaveformEditorWidgets/Private/SWaveformPanel.cpp
--- a/Engine/Plugins/Editor/WaveformEditor/Source/WaveformEditorWidgets/Private/SWaveformPanel.cpp
+++ b/Engine/Plugins/Editor/WaveformEditor/Source/WaveformEditorWidgets/Private/SWaveformPanel.cpp
@@ -9,6 +9,7 @@
 #include "WaveformEditorDisplayUnit.h"
 #include "WaveformEditorGridData.h"
 #include "WaveformEditorRenderData.h"
+#include "WaveformEditorStyle.h"
 #include "WaveformEditorTransportCoordinator.h"
 #include "WaveformEditorZoomController.h"
 #include "Widgets/SBoxPanel.h"
@@ -19,6 +20,9 @@ void SWaveformPanel::Construct(const FArguments& InArgs, TSharedRef<FWaveformEdi
 {
 	DisplayUnit = EWaveformEditorDisplayUnit::Seconds;
 
+	WaveformEditorStyle = &FWaveformEditorStyle::Get();
+	check(WaveformEditorStyle);
+
 	SetUpGridData(InRenderData, InTransportCoordinator);
 	SetUpWaveformViewer(InRenderData, InTransportCoordinator, InZoomManager, GridData.ToSharedRef());
 
@@ -73,20 +77,32 @@ void SWaveformPanel::CreateLayout()
 
 void SWaveformPanel::SetUpTimeRuler(TSharedRef<FWaveformEditorRenderData> InRenderData, TSharedRef<FWaveformEditorTransportCoordinator> InTransportCoordinator, TSharedRef<FWaveformEditorGridData> InGridData)
 {
-	TimeRuler = SNew(SWaveformEditorTimeRuler, InTransportCoordinator, InRenderData).DisplayUnit(DisplayUnit);
+	FWaveformEditorTimeRulerStyle* TimeRulerStyle = &WaveformEditorStyle->GetRegisteredWidgetStyle<FWaveformEditorTimeRulerStyle>("WaveformEditorRuler.Style").Get();
+	check(TimeRulerStyle);
+
+	TimeRuler = SNew(SWaveformEditorTimeRuler, InTransportCoordinator, InRenderData).DisplayUnit(DisplayUnit).Style(TimeRulerStyle);
+	TimeRulerStyle->OnStyleUpdated.AddSP(TimeRuler.ToSharedRef(), &SWaveformEditorTimeRuler::OnStyleUpdated);
 	InGridData->OnGridMetricsUpdated.AddSP(TimeRuler.ToSharedRef(), &SWaveformEditorTimeRuler::UpdateGridMetrics);
 	TimeRuler->OnTimeUnitMenuSelection.AddSP(this, &SWaveformPanel::UpdateDisplayUnit);
 }
 
 void SWaveformPanel::SetUpWaveformViewerOverlay(TSharedRef<FWaveformEditorTransportCoordinator> InTransportCoordinator, TSharedRef<FWaveformEditorZoomController> InZoomManager)
 {
-	WaveformViewerOverlay = SNew(SWaveformViewerOverlay, InTransportCoordinator, WaveformTransformationsOverlay.ToSharedRef());
+	FWaveformViewerOverlayStyle* OverlayStyle = &WaveformEditorStyle->GetRegisteredWidgetStyle<FWaveformViewerOverlayStyle>("WaveformViewerOverlay.Style").Get();
+	check(OverlayStyle);
+
+	WaveformViewerOverlay = SNew(SWaveformViewerOverlay, InTransportCoordinator, WaveformTransformationsOverlay.ToSharedRef()).Style(OverlayStyle);
+	OverlayStyle->OnStyleUpdated.AddSP(WaveformViewerOverlay.ToSharedRef(), &SWaveformViewerOverlay::OnStyleUpdated);
 	WaveformViewerOverlay->OnNewMouseDelta.BindSP(InZoomManager, &FWaveformEditorZoomController::ZoomByDelta);
 }
 
 void SWaveformPanel::SetUpWaveformViewer(TSharedRef<FWaveformEditorRenderData> InRenderData, TSharedRef<FWaveformEditorTransportCoordinator> InTransportCoordinator, TSharedRef<FWaveformEditorZoomController> InZoomManager, TSharedRef<FWaveformEditorGridData> InGridData)
 {
-	WaveformViewer = SNew(SWaveformViewer, InRenderData, InTransportCoordinator);
+	FWaveformViewerStyle* WaveViewerStyle = &WaveformEditorStyle->GetRegisteredWidgetStyle<FWaveformViewerStyle>("WaveformViewer.Style").Get();
+	check(WaveViewerStyle);
+	
+	WaveformViewer = SNew(SWaveformViewer, InRenderData, InTransportCoordinator).Style(WaveViewerStyle);
+	WaveViewerStyle->OnStyleUpdated.AddSP(WaveformViewer.ToSharedRef(), &SWaveformViewer::OnStyleUpdated);
 	InZoomManager->OnZoomRatioChanged.AddSP(InTransportCoordinator, &FWaveformEditorTransportCoordinator::OnZoomLevelChanged);
 	InGridData->OnGridMetricsUpdated.AddSP(WaveformViewer.ToSharedRef(), &SWaveformViewer::UpdateGridMetrics);
 }
diff --git a/Engine/Plugins/Editor/WaveformEditor/Source/WaveformEditorWidgets/Private/SWaveformViewer.cpp b/Engine/Plugins/Editor/WaveformEditor/Source/WaveformEditorWidgets/Private/SWaveformViewer.cpp
--- a/Engine/Plugins/Editor/WaveformEditor/Source/WaveformEditorWidgets/Private/SWaveformViewer.cpp
+++ b/Engine/Plugins/Editor/WaveformEditor/Source/WaveformEditorWidgets/Private/SWaveformViewer.cpp
@@ -17,15 +17,14 @@ void SWaveformViewer::Construct(const FArguments& InArgs, TSharedRef<FWaveformEd
 
 	Style = InArgs._Style;
 
-	check(Style)
-	WaveformColor = &Style->WaveformColor;
-	MajorGridLineColor = &Style->MajorGridLineColor;
-	MinorGridLineColor = &Style->MinorGridLineColor;
-	BackgroundColor = &Style->WaveformBackgroundColor;
-	BackgroundBrush = &Style->BackgroundBrush;
+	check(Style);
+	WaveformColor = Style->WaveformColor;
+	MajorGridLineColor = Style->MajorGridLineColor;
+	MinorGridLineColor = Style->MinorGridLineColor;
+	BackgroundColor = Style->WaveformBackgroundColor;
+	BackgroundBrush = Style->BackgroundBrush;
 	DesiredWidth = Style->DesiredWidth;
 	DesiredHeight = Style->DesiredHeight;
-
 }
 
 void SWaveformViewer::OnRenderDataUpdated()
@@ -51,9 +50,9 @@ int32 SWaveformViewer::OnPaint(const FPaintArgs& Args, const FGeometry& Allotted
 		OutDrawElements,
 		++LayerId,
 		AllottedGeometry.ToPaintGeometry(),
-		BackgroundBrush,
+		&BackgroundBrush,
 		ESlateDrawEffect::None, 
-		BackgroundColor->GetSpecifiedColor()
+		BackgroundColor.GetSpecifiedColor()
 	);
 
 	if (PixelWidth > 0)
@@ -74,7 +73,7 @@ int32 SWaveformViewer::OnPaint(const FPaintArgs& Args, const FGeometry& Allotted
 				AllottedGeometry.ToPaintGeometry(),
 				WaveformPoints,
 				ESlateDrawEffect::None, 
-				WaveformColor->GetSpecifiedColor()
+				WaveformColor.GetSpecifiedColor()
 			);
 		}
 
@@ -176,7 +175,7 @@ void SWaveformViewer::DrawGridLines(const FGeometry& AllottedGeometry, FSlateWin
 			AllottedGeometry.ToPaintGeometry(),
 			LinePoints,
 			ESlateDrawEffect::None,
-			MajorGridLineColor->GetSpecifiedColor(),
+			MajorGridLineColor.GetSpecifiedColor(),
 			false);
 
 
@@ -193,7 +192,7 @@ void SWaveformViewer::DrawGridLines(const FGeometry& AllottedGeometry, FSlateWin
 				AllottedGeometry.ToPaintGeometry(),
 				LinePoints,
 				ESlateDrawEffect::None,
-				MinorGridLineColor->GetSpecifiedColor(),
+				MinorGridLineColor.GetSpecifiedColor(),
 				false);
 
 		}
@@ -204,3 +203,22 @@ void SWaveformViewer::UpdateGridMetrics(const FWaveEditorGridMetrics& InMetrics)
 {
 	GridMetrics = InMetrics;
 }
+
+void SWaveformViewer::OnStyleUpdated(const FWaveformEditorWidgetStyleBase* UpdatedStyle)
+{
+	check(UpdatedStyle);
+	check(Style);
+
+	if (UpdatedStyle != Style)
+	{
+		return;
+	}
+
+	WaveformColor = Style->WaveformColor;
+	MajorGridLineColor = Style->MajorGridLineColor;
+	MinorGridLineColor = Style->MinorGridLineColor;
+	BackgroundColor = Style->WaveformBackgroundColor;
+	BackgroundBrush = Style->BackgroundBrush;
+	DesiredWidth = Style->DesiredWidth;
+	DesiredHeight = Style->DesiredHeight;
+}
diff --git a/Engine/Plugins/Editor/WaveformEditor/Source/WaveformEditorWidgets/Private/SWaveformViewer.h b/Engine/Plugins/Editor/WaveformEditor/Source/WaveformEditorWidgets/Private/SWaveformViewer.h
--- a/Engine/Plugins/Editor/WaveformEditor/Source/WaveformEditorWidgets/Private/SWaveformViewer.h
+++ b/Engine/Plugins/Editor/WaveformEditor/Source/WaveformEditorWidgets/Private/SWaveformViewer.h
@@ -24,9 +24,6 @@ public:
 
 	SLATE_BEGIN_ARGS(SWaveformViewer) 
 	{
-		const ISlateStyle* WaveEditorStyle = FSlateStyleRegistry::FindSlateStyle("WaveformEditorStyle");
-		check(WaveEditorStyle)
-		_Style = &WaveEditorStyle->GetWidgetStyle<FWaveformViewerStyle>("WaveformViewer.Style");
 	}
 
 	SLATE_STYLE_ARGUMENT(FWaveformViewerStyle, Style)
@@ -39,6 +36,7 @@ public:
 	void OnDisplayRangeUpdated(const TRange<float> NewDisplayRange);
 	void UpdateGridMetrics(const FWaveEditorGridMetrics& InMetrics);
 
+	void OnStyleUpdated(const FWaveformEditorWidgetStyleBase* UpdatedStyle);
 private:
 	virtual int32 OnPaint(const FPaintArgs& Args, const FGeometry& AllottedGeometry, const FSlateRect& MyCullingRect, FSlateWindowElementList& OutDrawElements, int32 LayerId, const FWidgetStyle& InWidgetStyle, bool bParentEnabled) const override;
 	virtual FVector2D ComputeDesiredSize(float) const override;
@@ -54,11 +52,12 @@ private:
 	FWaveEditorGridMetrics GridMetrics;
 
 	const FWaveformViewerStyle* Style = nullptr;
-	const FSlateBrush* BackgroundBrush = nullptr;
-	const FSlateColor* BackgroundColor = nullptr;
-	const FSlateColor* WaveformColor = nullptr;
-	const FSlateColor* MajorGridLineColor = nullptr;
-	const FSlateColor* MinorGridLineColor = nullptr;
+	FSlateBrush BackgroundBrush;
+	FSlateColor BackgroundColor = FLinearColor(0.02f, 0.02f, 0.02f, 1.f);
+	FSlateColor WaveformColor = FLinearColor::White;
+	FSlateColor MajorGridLineColor = FLinearColor::Black;
+	FSlateColor MinorGridLineColor = FLinearColor(0.f, 0.f, 0.f, 0.5f);
+	FSlateColor ZeroCrossingLineColor = FLinearColor::Black;
 	float DesiredHeight = 0.f;
 	float DesiredWidth = 0.f;
 
diff --git a/Engine/Plugins/Editor/WaveformEditor/Source/WaveformEditorWidgets/Private/SWaveformViewerOverlay.cpp b/Engine/Plugins/Editor/WaveformEditor/Source/WaveformEditorWidgets/Private/SWaveformViewerOverlay.cpp
--- a/Engine/Plugins/Editor/WaveformEditor/Source/WaveformEditorWidgets/Private/SWaveformViewerOverlay.cpp
+++ b/Engine/Plugins/Editor/WaveformEditor/Source/WaveformEditorWidgets/Private/SWaveformViewerOverlay.cpp
@@ -11,9 +11,9 @@ void SWaveformViewerOverlay::Construct(const FArguments& InArgs, TSharedRef<FWav
 	TransformationsOverlay = InTransformationsOverlay;
 
 	Style = InArgs._Style;
-
-	check(Style)
-	PlayheadColor = &Style->PlayheadColor;
+	
+	check(Style);
+	PlayheadColor = Style->PlayheadColor;
 	PlayheadWidth = Style->PlayheadWidth;
 	DesiredWidth = Style->DesiredWidth;
 	DesiredHeight = Style->DesiredHeight;
@@ -55,6 +55,22 @@ FCursorReply SWaveformViewerOverlay::OnCursorQuery(const FGeometry& MyGeometry,
 	return TransformationsOverlay->OnCursorQuery(MyGeometry, CursorEvent);
 }
 
+void SWaveformViewerOverlay::OnStyleUpdated(const FWaveformEditorWidgetStyleBase* UpdatedStyle)
+{
+	check(UpdatedStyle);
+	check(Style);
+
+	if (UpdatedStyle != Style)
+	{
+		return;
+	}
+
+	PlayheadColor = Style->PlayheadColor;
+	PlayheadWidth = Style->PlayheadWidth;
+	DesiredWidth = Style->DesiredWidth;
+	DesiredHeight = Style->DesiredHeight;
+}
+
 int32 SWaveformViewerOverlay::OnPaint(const FPaintArgs& Args, const FGeometry& AllottedGeometry, const FSlateRect& MyCullingRect, FSlateWindowElementList& OutDrawElements, int32 LayerId, const FWidgetStyle& InWidgetStyle, bool bParentEnabled) const
 {
 	LayerId = DrawPlayhead(AllottedGeometry, OutDrawElements, LayerId);
@@ -79,7 +95,7 @@ int32 SWaveformViewerOverlay::DrawPlayhead(const FGeometry& AllottedGeometry, FS
 		AllottedGeometry.ToPaintGeometry(),
 		LinePoints,
 		ESlateDrawEffect::None,
-		PlayheadColor->GetSpecifiedColor(),
+		PlayheadColor.GetSpecifiedColor(),
 		true, 
 		PlayheadWidth
 	);
diff --git a/Engine/Plugins/Editor/WaveformEditor/Source/WaveformEditorWidgets/Private/SWaveformViewerOverlay.h b/Engine/Plugins/Editor/WaveformEditor/Source/WaveformEditorWidgets/Private/SWaveformViewerOverlay.h
--- a/Engine/Plugins/Editor/WaveformEditor/Source/WaveformEditorWidgets/Private/SWaveformViewerOverlay.h
+++ b/Engine/Plugins/Editor/WaveformEditor/Source/WaveformEditorWidgets/Private/SWaveformViewerOverlay.h
@@ -16,10 +16,6 @@ class SWaveformViewerOverlay : public SLeafWidget
 public:
 	SLATE_BEGIN_ARGS(SWaveformViewerOverlay)
 	{
-		const ISlateStyle* WaveEditorStyle = FSlateStyleRegistry::FindSlateStyle("WaveformEditorStyle");
-		check(WaveEditorStyle)
-		_Style = &WaveEditorStyle->GetWidgetStyle<FWaveformViewerOverlayStyle>("WaveformViewerOverlay.Style");
-
 	}
 
 	SLATE_STYLE_ARGUMENT(FWaveformViewerOverlayStyle, Style)
@@ -33,6 +29,8 @@ public:
 	FReply OnMouseMove(const FGeometry& MyGeometry, const FPointerEvent& MouseEvent) override;
 	FCursorReply OnCursorQuery(const FGeometry& MyGeometry, const FPointerEvent& CursorEvent) const override;
 
+	void OnStyleUpdated(const FWaveformEditorWidgetStyleBase* UpdatedStyle);
+
 	FOnNewMouseDelta OnNewMouseDelta;
 
 private:
@@ -46,8 +44,8 @@ private:
 
 	const FWaveformViewerOverlayStyle* Style = nullptr;
 
-	const FSlateColor* PlayheadColor = nullptr;
-	float PlayheadWidth = 0.f;
+	FSlateColor PlayheadColor = FLinearColor(255.f, 0.1f, 0.2f, 1.f);
+	float PlayheadWidth = 1.0;
 	float DesiredWidth = 0.f;
 	float DesiredHeight = 0.f;
 };
\ No newline at end of file
diff --git a/Engine/Plugins/Editor/WaveformEditor/Source/WaveformEditorWidgets/Private/WaveformEditorSlateTypes.cpp b/Engine/Plugins/Editor/WaveformEditor/Source/WaveformEditorWidgets/Private/WaveformEditorSlateTypes.cpp
--- a/Engine/Plugins/Editor/WaveformEditor/Source/WaveformEditorWidgets/Private/WaveformEditorSlateTypes.cpp
+++ b/Engine/Plugins/Editor/WaveformEditor/Source/WaveformEditorWidgets/Private/WaveformEditorSlateTypes.cpp
@@ -9,9 +9,9 @@ namespace WaveformEditorStylesSharedParams
 {
 	const FLazyName BackgroundBrushName = TEXT("WhiteBrush");
 	const FLazyName HandleBrushName = TEXT("Sequencer.Timeline.VanillaScrubHandleDown");
-	const FLinearColor PlayheadColor = FLinearColor(255, 0.1, 0.2, 1.0);
-	const FLinearColor RulerTicksColor = FLinearColor(1.f, 1.f, 1.f, 0.9);
-	const FLinearColor WaveformBackGroundColor = FLinearColor(0.02, 0.02, 0.02, 1.f);
+	const FLinearColor PlayheadColor = FLinearColor(255.f, 0.1f, 0.2f, 1.f);
+	const FLinearColor RulerTicksColor = FLinearColor(1.f, 1.f, 1.f, 0.9f);
+	const FLinearColor WaveformBackGroundColor = FLinearColor(0.02f, 0.02f, 0.02f, 1.f);
 	const float ViewerHeight = 720.f;
 	const float ViewerWidth = 1280.f;
 }
@@ -21,7 +21,7 @@ const FName FWaveformViewerStyle::TypeName("FWaveformViewerStyle");
 FWaveformViewerStyle::FWaveformViewerStyle()
 	: WaveformColor(FLinearColor::White)
 	, MajorGridLineColor(FLinearColor::Black)
-	, MinorGridLineColor(FLinearColor(0.f, 0.f, 0.f, 0.5))
+	, MinorGridLineColor(FLinearColor(0.f, 0.f, 0.f, 0.5f))
 	, WaveformBackgroundColor(WaveformEditorStylesSharedParams::WaveformBackGroundColor)
 	, BackgroundBrush(*FAppStyle::GetBrush(WaveformEditorStylesSharedParams::BackgroundBrushName))
 	, DesiredWidth(WaveformEditorStylesSharedParams::ViewerWidth)
diff --git a/Engine/Plugins/Editor/WaveformEditor/Source/WaveformEditorWidgets/Private/WaveformEditorStyle.cpp b/Engine/Plugins/Editor/WaveformEditor/Source/WaveformEditorWidgets/Private/WaveformEditorStyle.cpp
--- a/Engine/Plugins/Editor/WaveformEditor/Source/WaveformEditorWidgets/Private/WaveformEditorStyle.cpp
+++ b/Engine/Plugins/Editor/WaveformEditor/Source/WaveformEditorWidgets/Private/WaveformEditorStyle.cpp
@@ -24,9 +24,9 @@ FWaveformEditorStyle::~FWaveformEditorStyle()
 	FSlateStyleRegistry::UnRegisterSlateStyle(*this);
 }
 
-const FWaveformEditorStyle& FWaveformEditorStyle::Get()
+FWaveformEditorStyle& FWaveformEditorStyle::Get()
 {
-	check(StyleInstance)
+	check(StyleInstance);
 	return *StyleInstance;
 }
 
@@ -74,7 +74,7 @@ void FWaveformEditorStyle::Init()
 const UWaveformEditorWidgetsSettings* FWaveformEditorStyle::GetWidgetsSettings() 
 {
 	const UWaveformEditorWidgetsSettings* WaveformEditorWidgetsSettings = GetDefault<UWaveformEditorWidgetsSettings>();
-	check(WaveformEditorWidgetsSettings)
+	check(WaveformEditorWidgetsSettings);
 
 	return WaveformEditorWidgetsSettings;
 }
@@ -122,6 +122,4 @@ void FWaveformEditorStyle::OnWidgetSettingsUpdated(const FName& PropertyName, co
 	{
 		WaveformEditorTimeRulerStyle->SetFontSize(Settings->RulerFontSize);
 	}
-}
-
-
+}
\ No newline at end of file
diff --git a/Engine/Plugins/Editor/WaveformEditor/Source/WaveformEditorWidgets/Private/WaveformEditorWidgetsSettings.cpp b/Engine/Plugins/Editor/WaveformEditor/Source/WaveformEditorWidgets/Private/WaveformEditorWidgetsSettings.cpp
--- a/Engine/Plugins/Editor/WaveformEditor/Source/WaveformEditorWidgets/Private/WaveformEditorWidgetsSettings.cpp
+++ b/Engine/Plugins/Editor/WaveformEditor/Source/WaveformEditorWidgets/Private/WaveformEditorWidgetsSettings.cpp
@@ -4,14 +4,14 @@
 
 UWaveformEditorWidgetsSettings::UWaveformEditorWidgetsSettings(const FObjectInitializer& ObjectInitializer)
 	: Super(ObjectInitializer)
-	, PlayheadColor(FLinearColor(255, 0.1, 0.2, 1.0))
+	, PlayheadColor(FLinearColor(255.f, 0.1f, 0.2f, 1.f))
 	, WaveformColor(FLinearColor::White)
-	, WaveformBackgroundColor(FLinearColor(0.02, 0.02, 0.02, 1.f))
+	, WaveformBackgroundColor(FLinearColor(0.02f, 0.02f, 0.02f, 1.f))
 	, MajorGridColor(FLinearColor::Black)
-	, MinorGridColor(FLinearColor(0.f, 0.f, 0.f, 0.5))
+	, MinorGridColor(FLinearColor(0.f, 0.f, 0.f, 0.5f))
 	, RulerBackgroundColor(FLinearColor::Black)
-	, RulerTicksColor(FLinearColor(1.f, 1.f, 1.f, 0.9))
-	, RulerTextColor(FLinearColor(1.f, 1.f, 1.f, 0.9))
+	, RulerTicksColor(FLinearColor(1.f, 1.f, 1.f, 0.9f))
+	, RulerTextColor(FLinearColor(1.f, 1.f, 1.f, 0.9f))
 	, RulerFontSize(10.f)
 {
 }
diff --git a/Engine/Plugins/Editor/WaveformEditor/Source/WaveformEditorWidgets/Public/SWaveformPanel.h b/Engine/Plugins/Editor/WaveformEditor/Source/WaveformEditorWidgets/Public/SWaveformPanel.h
--- a/Engine/Plugins/Editor/WaveformEditor/Source/WaveformEditorWidgets/Public/SWaveformPanel.h
+++ b/Engine/Plugins/Editor/WaveformEditor/Source/WaveformEditorWidgets/Public/SWaveformPanel.h
@@ -5,15 +5,16 @@
 #include "HAL/Platform.h"
 #include "Widgets/SCompoundWidget.h"
 
-enum class EWaveformEditorDisplayUnit;
 class FWaveformEditorGridData;
 class FWaveformEditorRenderData;
+class FWaveformEditorStyle;
 class FWaveformEditorTransportCoordinator;
 class FWaveformEditorZoomController;
 class SWaveformEditorTimeRuler;
 class SWaveformTransformationsOverlay;
 class SWaveformViewer;
 class SWaveformViewerOverlay;
+enum class EWaveformEditorDisplayUnit;
 
 class WAVEFORMEDITORWIDGETS_API SWaveformPanel : public SCompoundWidget
 {
@@ -45,4 +46,6 @@ private:
 	float CachedPixelWidth = 0.f;
 
 	EWaveformEditorDisplayUnit DisplayUnit;
+
+	FWaveformEditorStyle* WaveformEditorStyle;
 };
\ No newline at end of file
diff --git a/Engine/Plugins/Editor/WaveformEditor/Source/WaveformEditorWidgets/Public/WaveformEditorSlateTypes.h b/Engine/Plugins/Editor/WaveformEditor/Source/WaveformEditorWidgets/Public/WaveformEditorSlateTypes.h
--- a/Engine/Plugins/Editor/WaveformEditor/Source/WaveformEditorWidgets/Public/WaveformEditorSlateTypes.h
+++ b/Engine/Plugins/Editor/WaveformEditor/Source/WaveformEditorWidgets/Public/WaveformEditorSlateTypes.h
@@ -2,17 +2,32 @@
 
 #pragma once
 
-
+#include "Delegates/Delegate.h"
 #include "Styling/SlateStyle.h"
 #include "Styling/SlateWidgetStyle.h"
 
 #include "WaveformEditorSlateTypes.generated.h"
 
+DECLARE_MULTICAST_DELEGATE_OneParam(FOnWidgetStyleUpdated, const FWaveformEditorWidgetStyleBase* /*Updated Widget Style*/);
+
+USTRUCT()
+struct FWaveformEditorWidgetStyleBase : public FSlateWidgetStyle
+{
+	GENERATED_USTRUCT_BODY()
+
+	FWaveformEditorWidgetStyleBase() = default;
+	virtual ~FWaveformEditorWidgetStyleBase() = default;
+
+	virtual void BroadcastStyleUpdate() const { OnStyleUpdated.Broadcast(this); }
+
+	FOnWidgetStyleUpdated OnStyleUpdated;
+};
+
 /**
  * Represents the appearance of a Waveform Viewer
  */
 USTRUCT(BlueprintType)
-struct FWaveformViewerStyle : public FSlateWidgetStyle
+struct FWaveformViewerStyle : public FWaveformEditorWidgetStyleBase
 {
 	GENERATED_USTRUCT_BODY()
 
@@ -26,38 +41,38 @@ struct FWaveformViewerStyle : public FSlateWidgetStyle
 
 	UPROPERTY(EditAnywhere, BlueprintReadWrite, Category = Appearance)
 	FSlateColor WaveformColor;
-	FWaveformViewerStyle& SetWaveformColor(const FSlateColor InWaveformColor) { WaveformColor = InWaveformColor; return *this; }
+	FWaveformViewerStyle& SetWaveformColor(const FSlateColor InWaveformColor) { WaveformColor = InWaveformColor; BroadcastStyleUpdate(); return *this; }
 
 	UPROPERTY(EditAnywhere, BlueprintReadWrite, Category = Appearance)
 	FSlateColor MajorGridLineColor;
-	FWaveformViewerStyle& SetMajorGridLineColor(const FSlateColor InMajorGridLineColor) { MajorGridLineColor = InMajorGridLineColor; return *this; }
+	FWaveformViewerStyle& SetMajorGridLineColor(const FSlateColor InMajorGridLineColor) { MajorGridLineColor = InMajorGridLineColor; BroadcastStyleUpdate(); return *this; }
 
 	UPROPERTY(EditAnywhere, BlueprintReadWrite, Category = Appearance)
 	FSlateColor MinorGridLineColor;
-	FWaveformViewerStyle& SetMinorGridLineColor(const FSlateColor InMinorGridLineColor) { MinorGridLineColor = InMinorGridLineColor; return *this; }
+	FWaveformViewerStyle& SetMinorGridLineColor(const FSlateColor InMinorGridLineColor) { MinorGridLineColor = InMinorGridLineColor; BroadcastStyleUpdate(); return *this; }
 
 	UPROPERTY(EditAnywhere, BlueprintReadWrite, Category = Appearance)
 	FSlateColor WaveformBackgroundColor;
-	FWaveformViewerStyle& SetBackgroundColor(const FSlateColor InBackgroundColor) { WaveformBackgroundColor = InBackgroundColor; return *this; }
+	FWaveformViewerStyle& SetBackgroundColor(const FSlateColor InBackgroundColor) { WaveformBackgroundColor = InBackgroundColor; BroadcastStyleUpdate(); return *this; }
 
 	UPROPERTY(EditAnywhere, BlueprintReadWrite, Category = Appearance)
 	FSlateBrush BackgroundBrush;
-	FWaveformViewerStyle& SetBackgroundBrush(const FSlateBrush InBackgroundBrush) { BackgroundBrush = InBackgroundBrush; return *this; }
+	FWaveformViewerStyle& SetBackgroundBrush(const FSlateBrush InBackgroundBrush) { BackgroundBrush = InBackgroundBrush; BroadcastStyleUpdate(); return *this; }
 	
 	UPROPERTY(EditAnywhere, BlueprintReadWrite, Category = Appearance)
 	float DesiredWidth;
-	FWaveformViewerStyle& SetDesiredWidth(const float InDesiredWidth) { DesiredWidth = InDesiredWidth; return *this; }
+	FWaveformViewerStyle& SetDesiredWidth(const float InDesiredWidth) { DesiredWidth = InDesiredWidth; BroadcastStyleUpdate(); return *this; }
 
 	UPROPERTY(EditAnywhere, BlueprintReadWrite, Category = Appearance)
 	float DesiredHeight;
-	FWaveformViewerStyle& SetDesiredHeight(const float InDesiredHeight) { DesiredHeight = InDesiredHeight; return *this; }
+	FWaveformViewerStyle& SetDesiredHeight(const float InDesiredHeight) { DesiredHeight = InDesiredHeight; BroadcastStyleUpdate(); return *this; }
 };
 
 /**
  * Represents the appearance of a Waveform Viewer Overlay style
  */
 USTRUCT(BlueprintType)
-struct FWaveformViewerOverlayStyle : public FSlateWidgetStyle
+struct FWaveformViewerOverlayStyle : public FWaveformEditorWidgetStyleBase
 {
 	GENERATED_USTRUCT_BODY()
 
@@ -70,26 +85,26 @@ struct FWaveformViewerOverlayStyle : public FSlateWidgetStyle
 
 	UPROPERTY(EditAnywhere, BlueprintReadWrite, Category = Appearance)
 	FSlateColor PlayheadColor;
-	FWaveformViewerOverlayStyle& SetPlayheadColor(const FSlateColor InPlayheadColor) { PlayheadColor = InPlayheadColor; return *this; }
+	FWaveformViewerOverlayStyle& SetPlayheadColor(const FSlateColor InPlayheadColor) { PlayheadColor = InPlayheadColor; BroadcastStyleUpdate(); return *this; }
 
 	UPROPERTY(EditAnywhere, BlueprintReadWrite, Category = Appearance)
 	float PlayheadWidth;
-	FWaveformViewerOverlayStyle& SetPlayheadWidth(const float InPlayheadWidth) { PlayheadWidth = InPlayheadWidth; return *this; }
+	FWaveformViewerOverlayStyle& SetPlayheadWidth(const float InPlayheadWidth) { PlayheadWidth = InPlayheadWidth; BroadcastStyleUpdate(); return *this; }
 
 	UPROPERTY(EditAnywhere, BlueprintReadWrite, Category = Appearance)
 	float DesiredWidth;
-	FWaveformViewerOverlayStyle& SetDesiredWidth(const float InDesiredWidth) { DesiredWidth = InDesiredWidth; return *this; }
+	FWaveformViewerOverlayStyle& SetDesiredWidth(const float InDesiredWidth) { DesiredWidth = InDesiredWidth; BroadcastStyleUpdate(); return *this; }
 
 	UPROPERTY(EditAnywhere, BlueprintReadWrite, Category = Appearance)
 	float DesiredHeight;
-	FWaveformViewerOverlayStyle& SetDesiredHeight(const float InDesiredHeight) { DesiredHeight = InDesiredHeight; return *this; }
+	FWaveformViewerOverlayStyle& SetDesiredHeight(const float InDesiredHeight) { DesiredHeight = InDesiredHeight; BroadcastStyleUpdate(); return *this; }
 };
 
 /**
  * Represents the appearance of a Waveform Editor Time Ruler
  */
 USTRUCT(BlueprintType)
-struct FWaveformEditorTimeRulerStyle : public FSlateWidgetStyle
+struct FWaveformEditorTimeRulerStyle : public FWaveformEditorWidgetStyleBase
 {
 	GENERATED_USTRUCT_BODY()
 
@@ -103,46 +118,46 @@ struct FWaveformEditorTimeRulerStyle : public FSlateWidgetStyle
 
 	UPROPERTY(EditAnywhere, BlueprintReadWrite, Category = Appearance)
 	float HandleWidth;
-	FWaveformEditorTimeRulerStyle& SetHandleWidth(const float InHandleWidth) { HandleWidth = InHandleWidth; return *this; }
+	FWaveformEditorTimeRulerStyle& SetHandleWidth(const float InHandleWidth) { HandleWidth = InHandleWidth; BroadcastStyleUpdate(); return *this; }
 
 	UPROPERTY(EditAnywhere, BlueprintReadWrite, Category = Appearance)
 	FSlateColor HandleColor;
-	FWaveformEditorTimeRulerStyle& SetHandleColor(const FSlateColor& InHandleColor) { HandleColor = InHandleColor; return *this; }
+	FWaveformEditorTimeRulerStyle& SetHandleColor(const FSlateColor& InHandleColor) { HandleColor = InHandleColor; BroadcastStyleUpdate(); return *this; }
 
 	UPROPERTY(EditAnywhere, BlueprintReadWrite, Category = Appearance)
 	FSlateBrush HandleBrush;
-	FWaveformEditorTimeRulerStyle& SetHandleBrush(const FSlateBrush& InHandleBrush) { HandleBrush = InHandleBrush; return *this; }
+	FWaveformEditorTimeRulerStyle& SetHandleBrush(const FSlateBrush& InHandleBrush) { HandleBrush = InHandleBrush; BroadcastStyleUpdate(); return *this; }
 
 	UPROPERTY(EditAnywhere, BlueprintReadWrite, Category = Appearance)
 	FSlateColor TicksColor;
-	FWaveformEditorTimeRulerStyle& SetTicksColor(const FSlateColor& InTicksColor) { TicksColor = InTicksColor; return *this; }
+	FWaveformEditorTimeRulerStyle& SetTicksColor(const FSlateColor& InTicksColor) { TicksColor = InTicksColor; BroadcastStyleUpdate(); return *this; }
 
 	UPROPERTY(EditAnywhere, BlueprintReadWrite, Category = Appearance)
 	FSlateColor TicksTextColor;
-	FWaveformEditorTimeRulerStyle& SetTicksTextColor(const FSlateColor& InTicksTextColor) { TicksTextColor = InTicksTextColor; return *this; }
+	FWaveformEditorTimeRulerStyle& SetTicksTextColor(const FSlateColor& InTicksTextColor) { TicksTextColor = InTicksTextColor; BroadcastStyleUpdate(); return *this; }
 
 	UPROPERTY(EditAnywhere, BlueprintReadWrite, Category = Appearance)
 	FSlateFontInfo TicksTextFont;
-	FWaveformEditorTimeRulerStyle& SetTicksTextFont(const FSlateFontInfo& InTicksTextFont) { TicksTextFont = InTicksTextFont; return *this; }
-	FWaveformEditorTimeRulerStyle& SetFontSize(const float InFontSize) {TicksTextFont.Size = InFontSize; return *this;	}
+	FWaveformEditorTimeRulerStyle& SetTicksTextFont(const FSlateFontInfo& InTicksTextFont) { TicksTextFont = InTicksTextFont; BroadcastStyleUpdate(); return *this; }
+	FWaveformEditorTimeRulerStyle& SetFontSize(const float InFontSize) {TicksTextFont.Size = InFontSize; BroadcastStyleUpdate(); return *this;	}
 
 	UPROPERTY(EditAnywhere, BlueprintReadWrite, Category = Appearance)
 	float TicksTextOffset;
-	FWaveformEditorTimeRulerStyle& SetTicksTextOffset(const float InTicksTextOffset) { TicksTextOffset = InTicksTextOffset; return *this; }
+	FWaveformEditorTimeRulerStyle& SetTicksTextOffset(const float InTicksTextOffset) { TicksTextOffset = InTicksTextOffset; BroadcastStyleUpdate(); return *this; }
 
 	UPROPERTY(EditAnywhere, BlueprintReadWrite, Category = Appearance)
 	FSlateColor BackgroundColor;
-	FWaveformEditorTimeRulerStyle& SetBackgroundColor(const FSlateColor& InBackgroundColor) { BackgroundColor = InBackgroundColor; return *this; }
+	FWaveformEditorTimeRulerStyle& SetBackgroundColor(const FSlateColor& InBackgroundColor) { BackgroundColor = InBackgroundColor; BroadcastStyleUpdate(); return *this; }
 
 	UPROPERTY(EditAnywhere, BlueprintReadWrite, Category = Appearance)
 	FSlateBrush BackgroundBrush;
-	FWaveformEditorTimeRulerStyle& SetBackgroundBrush(const FSlateBrush& InBackgroundBrush) { BackgroundBrush = InBackgroundBrush; return *this; }
+	FWaveformEditorTimeRulerStyle& SetBackgroundBrush(const FSlateBrush& InBackgroundBrush) { BackgroundBrush = InBackgroundBrush; BroadcastStyleUpdate(); return *this; }
 
 	UPROPERTY(EditAnywhere, BlueprintReadWrite, Category = Appearance)
 	float DesiredWidth;
-	FWaveformEditorTimeRulerStyle& SetDesiredWidth(const float InDesiredWidth) { DesiredWidth = InDesiredWidth; return *this; }
+	FWaveformEditorTimeRulerStyle& SetDesiredWidth(const float InDesiredWidth) { DesiredWidth = InDesiredWidth; BroadcastStyleUpdate(); return *this; }
 
 	UPROPERTY(EditAnywhere, BlueprintReadWrite, Category = Appearance)
 	float DesiredHeight;
-	FWaveformEditorTimeRulerStyle& SetDesiredHeight(const float InDesiredHeight) { DesiredHeight = InDesiredHeight; return *this; }
+	FWaveformEditorTimeRulerStyle& SetDesiredHeight(const float InDesiredHeight) { DesiredHeight = InDesiredHeight; BroadcastStyleUpdate(); return *this; }
 };
\ No newline at end of file
diff --git a/Engine/Plugins/Editor/WaveformEditor/Source/WaveformEditorWidgets/Public/WaveformEditorStyle.h b/Engine/Plugins/Editor/WaveformEditor/Source/WaveformEditorWidgets/Public/WaveformEditorStyle.h
--- a/Engine/Plugins/Editor/WaveformEditor/Source/WaveformEditorWidgets/Public/WaveformEditorStyle.h
+++ b/Engine/Plugins/Editor/WaveformEditor/Source/WaveformEditorWidgets/Public/WaveformEditorStyle.h
@@ -18,14 +18,8 @@ public:
 	static FName StyleName;
 
 	/** Access the singleton instance for this style set */
-	static const FWaveformEditorStyle& Get();
+	static FWaveformEditorStyle& Get();
 	static void Init();
-	
-private:
-	static TUniquePtr<FWaveformEditorStyle> StyleInstance;
-
-	static const UWaveformEditorWidgetsSettings* GetWidgetsSettings();
-	static void  OnWidgetSettingsUpdated(const FName& PropertyName, const UWaveformEditorWidgetsSettings* Settings);
 
 	template< typename SlateWidgetStyle >
 	static TSharedRef<SlateWidgetStyle> GetRegisteredWidgetStyle(const FName& InStyleName)
@@ -35,4 +29,11 @@ private:
 		TSharedRef<SlateWidgetStyle> WidgetStyle = StaticCastSharedRef<SlateWidgetStyle>(*WidgetStyleBase);
 		return WidgetStyle;
 	}
+
+	
+private:
+	static TUniquePtr<FWaveformEditorStyle> StyleInstance;
+
+	static const UWaveformEditorWidgetsSettings* GetWidgetsSettings();
+	static void  OnWidgetSettingsUpdated(const FName& PropertyName, const UWaveformEditorWidgetsSettings* Settings);
 };
\ No newline at end of file
