[Waveform Editor] Allow users to specify a default transformation chain when opening the editor

diff --git a/Engine/Plugins/Editor/WaveformEditor/Source/WaveformEditor/Private/WaveformEditor.cpp b/Engine/Plugins/Editor/WaveformEditor/Source/WaveformEditor/Private/WaveformEditor.cpp
--- a/Engine/Plugins/Editor/WaveformEditor/Source/WaveformEditor/Private/WaveformEditor.cpp
+++ b/Engine/Plugins/Editor/WaveformEditor/Source/WaveformEditor/Private/WaveformEditor.cpp
@@ -15,6 +15,7 @@
 #include "WaveformEditorRenderData.h"
 #include "WaveformEditorStyle.h"
 #include "WaveformEditorToolMenuContext.h"
+#include "WaveformEditorTransformationsSettings.h"
 #include "WaveformEditorWaveWriter.h"
 #include "WaveformTransformationsRenderManager.h"
 #include "Widgets/Docking/SDockTab.h"
@@ -33,6 +34,7 @@ bool FWaveformEditor::Init(const EToolkitMode::Type Mode, const TSharedPtr<ITool
 	checkf(SoundWaveToEdit, TEXT("Tried to open a Soundwave Editor from a null soundwave"));
 
 	SoundWave = SoundWaveToEdit;
+	AddDefaultTransformations();
 
 	bool bIsInitialized = true;
 	
@@ -72,9 +74,34 @@ bool FWaveformEditor::Init(const EToolkitMode::Type Mode, const TSharedPtr<ITool
 	return bIsInitialized;
 }
 
+void FWaveformEditor::AddDefaultTransformations()
+{
+	if (!ensure(SoundWave))
+	{
+		return;
+	}
+
+	if (SoundWave->Transformations.Num() == 0)
+	{
+		const UWaveformEditorTransformationsSettings* TransformationsSettings = GetWaveformEditorTransformationsSettings();
+		for (const TSubclassOf<UWaveformTransformationBase> TransformationClass : TransformationsSettings->LaunchTransformations)
+		{
+			EObjectFlags MaskedOuterFlags = SoundWave ? SoundWave->GetMaskedFlags(RF_PropagateToSubObjects) : RF_NoFlags;
+
+			if (SoundWave->HasAnyFlags(RF_ClassDefaultObject | RF_ArchetypeObject))
+			{
+				MaskedOuterFlags |= RF_ArchetypeObject;
+			}
+
+			UWaveformTransformationBase* TransformationToAdd = NewObject<UWaveformTransformationBase>(SoundWave, TransformationClass.Get(), NAME_None, MaskedOuterFlags);
+			SoundWave->Transformations.Add(TransformationToAdd);
+		}
+	}
+}
+
 bool FWaveformEditor::SetupAudioComponent()
 {
-	if (SoundWave == nullptr)
+	if (!ensure(SoundWave))
 	{
 		return false;
 	}
@@ -108,7 +135,7 @@ bool FWaveformEditor::SetupAudioComponent()
 
 bool FWaveformEditor::SetUpTransportController()
 {
-	if (AudioComponent == nullptr)
+	if (!ensure(AudioComponent))
 	{
 		UE_LOG(LogWaveformEditor, Warning, TEXT("Trying to setup transport controls with a null audio component"));
 		return false;
@@ -126,7 +153,7 @@ bool FWaveformEditor::SetUpZoom()
 
 bool FWaveformEditor::BindDelegates()
 {
-	if (AudioComponent == nullptr)
+	if (!ensure(AudioComponent))
 	{
 		UE_LOG(LogWaveformEditor, Warning, TEXT("Failed to bind to playback percentage change, audio component is null"));
 		return false;
@@ -440,7 +467,7 @@ void FWaveformEditor::InitToolMenuContext(FToolMenuContext& MenuContext)
 
 bool FWaveformEditor::SetUpDetailsViews()
 {
-	if (SoundWave == nullptr)
+	if (!ensure(SoundWave)) 
 	{
 		UE_LOG(LogWaveformEditor, Warning, TEXT("Trying to setup wav editor properties view from a null SoundWave"));
 		return false;
@@ -527,7 +554,7 @@ TSharedRef<SDockTab> FWaveformEditor::SpawnTab_Transformations(const FSpawnTabAr
 }
 bool FWaveformEditor::SetUpWaveformPanel()
 {
-	if (SoundWave == nullptr)
+	if (!ensure(SoundWave))
 	{
 		UE_LOG(LogWaveformEditor, Warning, TEXT("Trying to setup waveform panel from a null SoundWave"));
 		return false;
@@ -651,7 +678,11 @@ bool FWaveformEditor::CanPressPlayButton() const
 
 bool FWaveformEditor::SetUpWaveWriter()
 {
-	check(SoundWave)
+	if (!ensure(SoundWave))
+	{
+		return false;
+	}
+
 	WaveWriter = MakeShared<FWaveformEditorWaveWriter>(SoundWave);
 	return WaveWriter != nullptr;
 }
@@ -662,4 +693,12 @@ void FWaveformEditor::ExportWaveform()
 	WaveWriter->ExportTransformedWaveform();
 }
 
+const UWaveformEditorTransformationsSettings* FWaveformEditor::GetWaveformEditorTransformationsSettings() const
+{
+	const UWaveformEditorTransformationsSettings* WaveformEditorTransformationsSettings = GetDefault<UWaveformEditorTransformationsSettings>();
+	check(WaveformEditorTransformationsSettings)
+
+	return WaveformEditorTransformationsSettings;
+}
+
 #undef LOCTEXT_NAMESPACE
diff --git a/Engine/Plugins/Editor/WaveformEditor/Source/WaveformEditor/Private/WaveformEditorTransformationsSettings.h b/Engine/Plugins/Editor/WaveformEditor/Source/WaveformEditor/Private/WaveformEditorTransformationsSettings.h
new file mode 100644
--- /dev/null
+++ b/Engine/Plugins/Editor/WaveformEditor/Source/WaveformEditor/Private/WaveformEditorTransformationsSettings.h
@@ -0,0 +1,43 @@
+// Copyright Epic Games, Inc. All Rights Reserved.
+
+#pragma once
+
+#include "Engine/DeveloperSettings.h"
+#include "IWaveformTransformation.h"
+
+#include "WaveformEditorTransformationsSettings.generated.h"
+
+/**
+ * Settings to regulate Waveform Transformations behavior inside Waveform Editor plugin.
+ */
+UCLASS(config = WaveformEditor, defaultconfig, meta = (DisplayName = "Waveform Editor Transformations"))
+class UWaveformEditorTransformationsSettings : public UDeveloperSettings
+{
+	GENERATED_BODY()
+
+//~ Begin UDeveloperSettings interface
+	virtual FName GetCategoryName() const
+	{
+		return TEXT("Plugins");
+	}
+
+#if WITH_EDITOR
+	virtual FText GetSectionText() const override
+	{
+		return NSLOCTEXT("WaveformEditorTransformations", "WaveformEditorTransformationsSettingsSection", "Waveform Editor Transformations");
+	}
+
+	virtual FName GetSectionName() const override
+	{
+		return TEXT("Waveform Editor Transformations");
+	}
+
+#endif
+	//~ End UDeveloperSettings interface
+	
+public:
+	/** A Transformation chain that will be added to the inspected Soundwave if there aren't any  */
+	UPROPERTY(config, EditAnywhere, Category = "Launch Options")
+	TArray<TSubclassOf<UWaveformTransformationBase>> LaunchTransformations;
+
+};
diff --git a/Engine/Plugins/Editor/WaveformEditor/Source/WaveformEditor/Public/WaveformEditor.h b/Engine/Plugins/Editor/WaveformEditor/Source/WaveformEditor/Public/WaveformEditor.h
--- a/Engine/Plugins/Editor/WaveformEditor/Source/WaveformEditor/Public/WaveformEditor.h
+++ b/Engine/Plugins/Editor/WaveformEditor/Source/WaveformEditor/Public/WaveformEditor.h
@@ -12,8 +12,9 @@
 
 class IToolkitHost;
 class SDockTab;
-class USoundWave;
 class UAudioComponent;
+class USoundWave;
+class UWaveformEditorTransformationsSettings;
 
 class WAVEFORMEDITOR_API FWaveformEditor 
 	: public FAssetEditorToolkit
@@ -85,6 +86,9 @@ private:
 	bool SetUpWaveWriter();
 	void ExportWaveform();
 
+	const UWaveformEditorTransformationsSettings* GetWaveformEditorTransformationsSettings() const;
+	void AddDefaultTransformations();
+
 	/** Waveform Preview widget */
 	TSharedPtr<class SWaveformPanel> WaveformPanel;
 
diff --git a/Engine/Plugins/Editor/WaveformEditor/Source/WaveformEditor/WaveformEditor.build.cs b/Engine/Plugins/Editor/WaveformEditor/Source/WaveformEditor/WaveformEditor.build.cs
--- a/Engine/Plugins/Editor/WaveformEditor/Source/WaveformEditor/WaveformEditor.build.cs
+++ b/Engine/Plugins/Editor/WaveformEditor/Source/WaveformEditor/WaveformEditor.build.cs
@@ -17,6 +17,7 @@ namespace UnrealBuildTool.Rules
 					"AssetRegistry",
 					"Core", 
 					"CoreUObject",
+					"DeveloperSettings",
 					"EditorStyle",
 					"Engine",
 					"InputCore",
