[WaveformEditor] Fix for  error C2036 in non unity builds caused by transformation settings
#trivial

diff --git a/Engine/Plugins/Editor/WaveformEditor/Source/WaveformEditor/Private/WaveformEditorTransformationsSettings.cpp b/Engine/Plugins/Editor/WaveformEditor/Source/WaveformEditor/Private/WaveformEditorTransformationsSettings.cpp
new file mode 100644
--- /dev/null
+++ b/Engine/Plugins/Editor/WaveformEditor/Source/WaveformEditor/Private/WaveformEditorTransformationsSettings.cpp
@@ -0,0 +1,22 @@
+// Copyright Epic Games, Inc. All Rights Reserved.
+
+#include "WaveformEditorTransformationsSettings.h"
+
+
+FName UWaveformEditorTransformationsSettings::GetCategoryName() const
+{
+	return TEXT("Plugins");
+}
+
+#if WITH_EDITOR
+
+FText UWaveformEditorTransformationsSettings::GetSectionText() const
+{
+	return NSLOCTEXT("WaveformEditorTransformations", "WaveformEditorTransformationsSettingsSection", "Waveform Editor Transformations");
+}
+
+FName UWaveformEditorTransformationsSettings::GetSectionName() const
+{
+	return TEXT("Waveform Editor Transformations");
+}
+#endif
\ No newline at end of file
diff --git a/Engine/Plugins/Editor/WaveformEditor/Source/WaveformEditor/Private/WaveformEditorTransformationsSettings.h b/Engine/Plugins/Editor/WaveformEditor/Source/WaveformEditor/Private/WaveformEditorTransformationsSettings.h
--- a/Engine/Plugins/Editor/WaveformEditor/Source/WaveformEditor/Private/WaveformEditorTransformationsSettings.h
+++ b/Engine/Plugins/Editor/WaveformEditor/Source/WaveformEditor/Private/WaveformEditorTransformationsSettings.h
@@ -2,8 +2,10 @@
 
 #pragma once
 
+#include "Containers/Array.h"
 #include "Engine/DeveloperSettings.h"
 #include "IWaveformTransformation.h"
+#include "Templates/SubclassOf.h"
 
 #include "WaveformEditorTransformationsSettings.generated.h"
 
@@ -16,22 +18,11 @@ class UWaveformEditorTransformationsSettings : public UDeveloperSettings
 	GENERATED_BODY()
 
 //~ Begin UDeveloperSettings interface
-	virtual FName GetCategoryName() const
-	{
-		return TEXT("Plugins");
-	}
+	virtual FName GetCategoryName() const;
 
 #if WITH_EDITOR
-	virtual FText GetSectionText() const override
-	{
-		return NSLOCTEXT("WaveformEditorTransformations", "WaveformEditorTransformationsSettingsSection", "Waveform Editor Transformations");
-	}
-
-	virtual FName GetSectionName() const override
-	{
-		return TEXT("Waveform Editor Transformations");
-	}
-
+	virtual FText GetSectionText() const override;
+	virtual FName GetSectionName() const override;
 #endif
 	//~ End UDeveloperSettings interface
 	
@@ -39,5 +30,4 @@ public:
 	/** A Transformation chain that will be added to the inspected Soundwave if there aren't any  */
 	UPROPERTY(config, EditAnywhere, Category = "Launch Options")
 	TArray<TSubclassOf<UWaveformTransformationBase>> LaunchTransformations;
-
 };
