[WaveformEditor] Change renedering button icon to use generic export one  #trivial # rb none #preflight none

diff --git a/Engine/Plugins/Editor/WaveformEditor/Source/WaveformEditor/Private/WaveformEditor.cpp b/Engine/Plugins/Editor/WaveformEditor/Source/WaveformEditor/Private/WaveformEditor.cpp
--- a/Engine/Plugins/Editor/WaveformEditor/Source/WaveformEditor/Private/WaveformEditor.cpp
+++ b/Engine/Plugins/Editor/WaveformEditor/Source/WaveformEditor/Private/WaveformEditor.cpp
@@ -270,7 +270,7 @@ bool FWaveformEditor::RegisterToolbar()
 			Commands.ExportWaveform,
 			LOCTEXT("WaveformEditorRender", ""),
 			LOCTEXT("WaveformEditorRenderButtonTooltip", "Exports the edited waveform to a USoundWave asset"),
-			FSlateIcon(FAppStyle::GetAppStyleSetName(), "GenericCurveEditor.SetViewModeAbsolute")
+			FSlateIcon(FAppStyle::GetAppStyleSetName(), "LevelEditor.ExportAll")
 		);	
 
 		ExportSection.AddEntry(ExportEntry);
@@ -286,7 +286,7 @@ bool FWaveformEditor::RegisterToolbar()
 						FNewToolMenuChoice(FOnGetContent::CreateSP(Context->WaveformEditor.Pin().Get(), &FWaveformEditor::GenerateExportOptionsMenu)),
 						LOCTEXT("ExportsOptions_Label", "Export Options"),
 						LOCTEXT("ExportsOptions_ToolTip", "Export options for this waveform"),
-						FSlateIcon(FAppStyle::GetAppStyleSetName(), "GenericCurveEditor.SetViewModeAbsolute"),
+						FSlateIcon(FAppStyle::GetAppStyleSetName(), "LevelEditor.ExportAll"),
 						true
 					);
 
