[WaveformEditor] Move settings to EditorPerProjectUserSettings.ini
#trivial
fix for #jira UE-164878

diff --git a/Engine/Plugins/Editor/WaveformEditor/Source/WaveformEditor/Private/WaveformEditorTransformationsSettings.h b/Engine/Plugins/Editor/WaveformEditor/Source/WaveformEditor/Private/WaveformEditorTransformationsSettings.h
--- a/Engine/Plugins/Editor/WaveformEditor/Source/WaveformEditor/Private/WaveformEditorTransformationsSettings.h
+++ b/Engine/Plugins/Editor/WaveformEditor/Source/WaveformEditor/Private/WaveformEditorTransformationsSettings.h
@@ -12,7 +12,7 @@
 /**
  * Settings to regulate Waveform Transformations behavior inside Waveform Editor plugin.
  */
-UCLASS(config = WaveformEditor, defaultconfig, meta = (DisplayName = "Waveform Editor Transformations"))
+UCLASS(config = EditorPerProjectUserSettings, defaultconfig, meta = (DisplayName = "Waveform Editor Transformations"))
 class UWaveformEditorTransformationsSettings : public UDeveloperSettings
 {
 	GENERATED_BODY()
diff --git a/Engine/Plugins/Editor/WaveformEditor/Source/WaveformEditorWidgets/Private/WaveformEditorWidgetsSettings.h b/Engine/Plugins/Editor/WaveformEditor/Source/WaveformEditorWidgets/Private/WaveformEditorWidgetsSettings.h
--- a/Engine/Plugins/Editor/WaveformEditor/Source/WaveformEditorWidgets/Private/WaveformEditorWidgetsSettings.h
+++ b/Engine/Plugins/Editor/WaveformEditor/Source/WaveformEditorWidgets/Private/WaveformEditorWidgetsSettings.h
@@ -8,7 +8,7 @@
 
 DECLARE_MULTICAST_DELEGATE_TwoParams(FOnWaveformEditorWidgetsSettingsChanged, const FName& /*Property Name*/, const UWaveformEditorWidgetsSettings*);
 
-UCLASS(config = WaveformEditor, defaultconfig, meta = (DisplayName = "Waveform Editor Display"))
+UCLASS(config = EditorPerProjectUserSettings, defaultconfig, meta = (DisplayName = "Waveform Editor Display"))
 class UWaveformEditorWidgetsSettings : public UDeveloperSettings
 {
 	GENERATED_BODY()
