[WaveformEditor] Allow user to specify empty transformations slots in the default chain
#trivial
fix for #jira UE-164742

diff --git a/Engine/Plugins/Editor/WaveformEditor/Source/WaveformEditor/Private/WaveformEditor.cpp b/Engine/Plugins/Editor/WaveformEditor/Source/WaveformEditor/Private/WaveformEditor.cpp
--- a/Engine/Plugins/Editor/WaveformEditor/Source/WaveformEditor/Private/WaveformEditor.cpp
+++ b/Engine/Plugins/Editor/WaveformEditor/Source/WaveformEditor/Private/WaveformEditor.cpp
@@ -86,15 +86,22 @@ void FWaveformEditor::AddDefaultTransformations()
 		const UWaveformEditorTransformationsSettings* TransformationsSettings = GetWaveformEditorTransformationsSettings();
 		for (const TSubclassOf<UWaveformTransformationBase> TransformationClass : TransformationsSettings->LaunchTransformations)
 		{
-			EObjectFlags MaskedOuterFlags = SoundWave ? SoundWave->GetMaskedFlags(RF_PropagateToSubObjects) : RF_NoFlags;
+			if (TransformationClass)
+			{
+				EObjectFlags MaskedOuterFlags = SoundWave ? SoundWave->GetMaskedFlags(RF_PropagateToSubObjects) : RF_NoFlags;
+
+				if (SoundWave->HasAnyFlags(RF_ClassDefaultObject | RF_ArchetypeObject))
+				{
+					MaskedOuterFlags |= RF_ArchetypeObject;
+				}
 
-			if (SoundWave->HasAnyFlags(RF_ClassDefaultObject | RF_ArchetypeObject))
+				UWaveformTransformationBase* TransformationToAdd = NewObject<UWaveformTransformationBase>(SoundWave, TransformationClass.Get(), NAME_None, MaskedOuterFlags);
+				SoundWave->Transformations.Add(TransformationToAdd);
+			}
+			else
 			{
-				MaskedOuterFlags |= RF_ArchetypeObject;
+				SoundWave->Transformations.Add(nullptr);
 			}
-
-			UWaveformTransformationBase* TransformationToAdd = NewObject<UWaveformTransformationBase>(SoundWave, TransformationClass.Get(), NAME_None, MaskedOuterFlags);
-			SoundWave->Transformations.Add(TransformationToAdd);
 		}
 	}
 }
