[Waveform Editor] Change to sample view when pixel length > num samples displayed

diff --git a/Engine/Plugins/Editor/WaveformEditor/Source/WaveformEditorWidgets/Private/SWaveformEditorTimeRuler.cpp b/Engine/Plugins/Editor/WaveformEditor/Source/WaveformEditorWidgets/Private/SWaveformEditorTimeRuler.cpp
--- a/Engine/Plugins/Editor/WaveformEditor/Source/WaveformEditorWidgets/Private/SWaveformEditorTimeRuler.cpp
+++ b/Engine/Plugins/Editor/WaveformEditor/Source/WaveformEditorWidgets/Private/SWaveformEditorTimeRuler.cpp
@@ -174,8 +174,8 @@ void SWaveformEditorTimeRuler::MakeTimeUnitsSubMenu(FMenuBuilder& SubMenuBuilder
 	);
 
 	SubMenuBuilder.AddMenuEntry(
-		LOCTEXT("TimeUnitsSubMenuEntry_Samples", "Samples"),
-		LOCTEXT("TimeUnitsSubMenuEntry_Samples", "Samples"),
+		LOCTEXT("TimeUnitsSubMenuEntry_Frames", "Frames"),
+		LOCTEXT("TimeUnitsSubMenuEntry_Frames", "Frames"),
 		FSlateIcon(),
 		FUIAction(
 			FExecuteAction::CreateSP(this, &SWaveformEditorTimeRuler::NotifyTimeUnitMenuSelection, EWaveformEditorDisplayUnit::Samples),
@@ -261,7 +261,7 @@ FReply SWaveformEditorTimeRuler::OnMouseButtonDown(const FGeometry& MyGeometry,
 
 FReply SWaveformEditorTimeRuler::OnMouseButtonUp(const FGeometry& MyGeometry, const FPointerEvent& MouseEvent)
 {
-	const bool HandleRightMouseButton = MouseEvent.GetEffectingButton() == EKeys::RightMouseButton && HasMouseCapture();
+	const bool HandleRightMouseButton = MouseEvent.GetEffectingButton() == EKeys::RightMouseButton;
 
 	if (HandleRightMouseButton)
 	{
diff --git a/Engine/Plugins/Editor/WaveformEditor/Source/WaveformEditorWidgets/Private/SWaveformViewer.cpp b/Engine/Plugins/Editor/WaveformEditor/Source/WaveformEditorWidgets/Private/SWaveformViewer.cpp
--- a/Engine/Plugins/Editor/WaveformEditor/Source/WaveformEditorWidgets/Private/SWaveformViewer.cpp
+++ b/Engine/Plugins/Editor/WaveformEditor/Source/WaveformEditorWidgets/Private/SWaveformViewer.cpp
@@ -1,7 +1,9 @@
 // Copyright Epic Games, Inc. All Rights Reserved.
 
 #include "SWaveformViewer.h"
+
 #include "WaveformEditorRenderData.h"
+#include "WaveformEditorSlateTypes.h"
 
 void SWaveformViewer::Construct(const FArguments& InArgs, TSharedRef<FWaveformEditorRenderData> InRenderData, TSharedRef<FWaveformEditorTransportCoordinator> InTransportCoordinator)
 {
@@ -16,7 +18,6 @@ void SWaveformViewer::Construct(const FArguments& InArgs, TSharedRef<FWaveformEd
 	DisplayRange = TRange<float>::Inclusive(0.f, RenderData->GetOriginalWaveformDurationInSeconds());
 
 	Style = InArgs._Style;
-
 	check(Style);
 	WaveformColor = Style->WaveformColor;
 	MajorGridLineColor = Style->MajorGridLineColor;
@@ -25,13 +26,15 @@ void SWaveformViewer::Construct(const FArguments& InArgs, TSharedRef<FWaveformEd
 	BackgroundBrush = Style->BackgroundBrush;
 	DesiredWidth = Style->DesiredWidth;
 	DesiredHeight = Style->DesiredHeight;
+	SampleMarkersSize = Style->SampleMarkersSize;
+	WaveformLineThickness = Style->WaveformLineThickness;
+	ZeroCrossingLineColor = Style->ZeroCrossingLineColor;
+	ZeroCrossingLineThickness = Style->ZeroCrossingLineThickness;
 }
-
 void SWaveformViewer::OnRenderDataUpdated()
 {
 	SampleData = RenderData->GetSampleData();
 	bForceRedraw = true;
-
 }
 
 void SWaveformViewer::OnDisplayRangeUpdated(const TRange<float> NewDisplayRange)
@@ -59,22 +62,69 @@ int32 SWaveformViewer::OnPaint(const FPaintArgs& Args, const FGeometry& Allotted
 	{
 		DrawGridLines(AllottedGeometry, OutDrawElements, LayerId);
 
-		TArray<FVector2D> WaveformPoints;
- 		WaveformPoints.SetNumUninitialized(2);
-
-		for (const FWaveformLineCoordinates& PeakCoordinates : CachedDrawCoordinates)
+		if (WaveformDrawMode == EWaveformDrawMode::BinnedPeaks)
 		{
-			WaveformPoints[0] = PeakCoordinates.PointA;
-			WaveformPoints[1] = PeakCoordinates.PointB;
+			TArray<FVector2D> BinDrawPoints;
+			BinDrawPoints.SetNumUninitialized(2);
+
+			for (const FSamplesBinCoordinates& BinCoordinates : CachedBinsDrawCoordinates)
+			{
+				BinDrawPoints[0] = BinCoordinates.PointA;
+				BinDrawPoints[1] = BinCoordinates.PointB;
+
+				FSlateDrawElement::MakeLines(
+					OutDrawElements,
+					LayerId,
+					AllottedGeometry.ToPaintGeometry(),
+					BinDrawPoints,
+					ESlateDrawEffect::None,
+					WaveformColor.GetSpecifiedColor()
+				);
+			}
 
-			FSlateDrawElement::MakeLines(
-				OutDrawElements,
-				LayerId,
-				AllottedGeometry.ToPaintGeometry(),
-				WaveformPoints,
-				ESlateDrawEffect::None, 
-				WaveformColor.GetSpecifiedColor()
-			);
+		}
+		else if (WaveformDrawMode == EWaveformDrawMode::WaveformLine)
+		{
+			TArray<FVector2D> SamplesDrawCoordinates;
+			const uint16 NChannels = RenderData->GetNumChannels();
+
+			for (uint16 Channel = 0; Channel < NChannels; ++Channel)
+			{
+				for (int32 SampleIndex = Channel; SampleIndex < CachedSampleDrawCoordinates.Num(); SampleIndex += NChannels)
+				{
+					SamplesDrawCoordinates.Add(CachedSampleDrawCoordinates[SampleIndex]);
+					
+					const FVector2D SampleBoxSize = FVector2D(SampleMarkersSize, SampleMarkersSize);
+					const FVector2D SampleBoxTransform = FVector2D(CachedSampleDrawCoordinates[SampleIndex] - (SampleBoxSize / 2.f));
+
+					const FPaintGeometry SampleBoxGeometry = AllottedGeometry.ToPaintGeometry(
+						SampleBoxSize,
+						FSlateLayoutTransform(SampleBoxTransform)
+					);
+
+					FSlateDrawElement::MakeBox(
+						OutDrawElements,
+						LayerId,
+						SampleBoxGeometry,
+						&BackgroundBrush,
+						ESlateDrawEffect::None,
+						WaveformColor.GetSpecifiedColor()
+					);
+				}
+
+				FSlateDrawElement::MakeLines(
+					OutDrawElements,
+					LayerId,
+					AllottedGeometry.ToPaintGeometry(),
+					SamplesDrawCoordinates,
+					ESlateDrawEffect::None,
+					WaveformColor.GetSpecifiedColor(),
+					true,
+					WaveformLineThickness
+				);
+
+				SamplesDrawCoordinates.Empty();
+			}
 		}
 
 		++LayerId;
@@ -100,18 +150,38 @@ void SWaveformViewer::Tick(const FGeometry& AllottedGeometry, const double InCur
 	{
 		CachedPixelWidth = DiscretePixelWidth;
 		CachedPixelHeight = AllottedGeometry.GetAbsoluteSize().Y;
-		WaveformDrawingUtils::GetBinnedPeaksFromWaveformRawData(CachedPeaks, CachedPixelWidth, SampleData.GetData(), RenderData->GetNumSamples(), RenderData->GetSampleRate(), RenderData->GetNumChannels(), DisplayRange.GetLowerBoundValue(), DisplayRange.GetUpperBoundValue());
-		GenerateWaveformLines(CachedDrawCoordinates, CachedPeaks, AllottedGeometry);
+		const float NumSamplesToDisplay = (DisplayRange.GetUpperBoundValue() - DisplayRange.GetLowerBoundValue()) * RenderData->GetSampleRate();
+		WaveformDrawMode = NumSamplesToDisplay > CachedPixelWidth ? EWaveformDrawMode::BinnedPeaks : EWaveformDrawMode::WaveformLine;
+
+		if (WaveformDrawMode == EWaveformDrawMode::BinnedPeaks)
+		{
+			WaveformDrawingUtils::GetBinnedPeaksFromWaveformRawData(CachedPeaks, CachedPixelWidth, SampleData.GetData(), RenderData->GetNumSamples(), RenderData->GetSampleRate(), RenderData->GetNumChannels(), DisplayRange.GetLowerBoundValue(), DisplayRange.GetUpperBoundValue());
+			GenerateWaveformBinsCoordinates(CachedBinsDrawCoordinates, CachedPeaks, AllottedGeometry);
+		}
+		else
+		{
+			GenerateWaveformSamplesCoordinates(CachedSampleDrawCoordinates, AllottedGeometry);
+		}
+		
 		bForceRedraw = false;
 	}
 	else if (CachedPixelHeight != AllottedGeometry.GetAbsoluteSize().Y)
 	{
 		CachedPixelHeight = AllottedGeometry.GetAbsoluteSize().Y;
-		GenerateWaveformLines(CachedDrawCoordinates, CachedPeaks, AllottedGeometry);
+
+		if (WaveformDrawMode == EWaveformDrawMode::BinnedPeaks)
+		{
+			GenerateWaveformBinsCoordinates(CachedBinsDrawCoordinates, CachedPeaks, AllottedGeometry);
+		}
+		else
+		{
+			GenerateWaveformSamplesCoordinates(CachedSampleDrawCoordinates, AllottedGeometry);
+		}
+
 	}
 }
 
-void SWaveformViewer::GenerateWaveformLines(TArray<FWaveformLineCoordinates>& OutDrawCoordinates, const TArray<WaveformDrawingUtils::SampleRange>& InWaveformPeaks, const FGeometry& InAllottedGeometry, const float VerticalZoomFactor)
+void SWaveformViewer::GenerateWaveformBinsCoordinates(TArray<FSamplesBinCoordinates>& OutDrawCoordinates, const TArray<WaveformDrawingUtils::SampleRange>& InWaveformPeaks, const FGeometry& InAllottedGeometry, const float VerticalZoomFactor /*= 1.f*/)
 {
 	const int NChannels = RenderData->GetNumChannels();
 
@@ -120,34 +190,28 @@ void SWaveformViewer::GenerateWaveformLines(TArray<FWaveformLineCoordinates>& Ou
 		return;
 	}
 
-	const float ChannelSlotHeight = InAllottedGeometry.GetAbsoluteSize().Y / NChannels;
-	const float ChannelSlotMidPoint = ChannelSlotHeight / 2;
-	const float HeightScale = InAllottedGeometry.GetAbsoluteSize().Y / (2.f * 32767 * NChannels) * VerticalZoomFactor;
-	
 	const uint32 PixelWidth = FMath::FloorToInt(InAllottedGeometry.GetAbsoluteSize().X);
 	check(PixelWidth * NChannels == InWaveformPeaks.Num());
 
-	float MaxDistanceFromBoundary = 2.f;
+	const float HeightScale = InAllottedGeometry.GetAbsoluteSize().Y / (2.f * TNumericLimits<int16>::Max() * NChannels) * MaxWaveformHeight * VerticalZoomFactor;
 
 	OutDrawCoordinates.SetNumUninitialized(InWaveformPeaks.Num());
-	FWaveformLineCoordinates* OutCoordinatesData = OutDrawCoordinates.GetData();
+	FSamplesBinCoordinates* OutCoordinatesData = OutDrawCoordinates.GetData();
 
 	for (uint16 Channel = 0; Channel < NChannels; ++Channel)
 	{
-		const float WaveformTopBoundary = ChannelSlotHeight * Channel;
-		const float WaveformBottomBoundary = WaveformTopBoundary + ChannelSlotHeight;
-		const float WaveformDrawCenter = WaveformTopBoundary + ChannelSlotMidPoint;
+		const FChannelSlotBoundaries ChannelBoundaries(Channel, NChannels, InAllottedGeometry);
 
 		for (uint32 Pixel = 0; Pixel < PixelWidth; ++Pixel)
 		{
 			uint32 PeakIndex = Pixel * NChannels + Channel;
 			const WaveformDrawingUtils::SampleRange& SamplePeaks = InWaveformPeaks[PeakIndex];
 
-			const float SampleMaxScaled = SamplePeaks.GetUpperBoundValue() * HeightScale > 0.001f ? SamplePeaks.GetUpperBoundValue() * HeightScale : 0.1f;
-			const float SampleMinScaled = SamplePeaks.GetLowerBoundValue() * HeightScale < -0.001f ? SamplePeaks.GetLowerBoundValue() * HeightScale : -0.1f;
+			const float SampleMaxScaled = SamplePeaks.GetUpperBoundValue() * HeightScale > MinScaledSamplePeakValue ? SamplePeaks.GetUpperBoundValue() * HeightScale : MinWaveformBinHeight;
+			const float SampleMinScaled = SamplePeaks.GetLowerBoundValue() * HeightScale < -1 * MinScaledSamplePeakValue ? SamplePeaks.GetLowerBoundValue() * HeightScale : -1.f * MinWaveformBinHeight;
 
-			const float Top = FMath::Max(WaveformDrawCenter - SampleMaxScaled, WaveformTopBoundary + MaxDistanceFromBoundary);
-			const float Bottom = FMath::Min(WaveformDrawCenter - SampleMinScaled, WaveformBottomBoundary - MaxDistanceFromBoundary);
+			const float Top = FMath::Max(ChannelBoundaries.Center - SampleMaxScaled, ChannelBoundaries.Top + MaxDistanceFromChannelBoundaries);
+			const float Bottom = FMath::Min(ChannelBoundaries.Center - SampleMinScaled, ChannelBoundaries.Bottom - MaxDistanceFromChannelBoundaries);
 
 			OutCoordinatesData[PeakIndex].PointA = FVector2D(Pixel, Top);
 			OutCoordinatesData[PeakIndex].PointB = FVector2D(Pixel, Bottom);
@@ -155,6 +219,56 @@ void SWaveformViewer::GenerateWaveformLines(TArray<FWaveformLineCoordinates>& Ou
 	}
 }
 
+void SWaveformViewer::GenerateWaveformSamplesCoordinates(TArray<FVector2D>& OutDrawCoordinates, const FGeometry& InAllottedGeometry, const float VerticalZoomFactor /*= 1.f*/)
+{
+	
+	OutDrawCoordinates.Empty();
+	const uint16 NChannels = RenderData->GetNumChannels();
+
+	if (!ensure(NChannels != 0))
+	{
+		return;
+	}
+
+	const uint32 TotalWaveformFrames = RenderData->GetNumSamples() / NChannels;
+	double NumFramesToDisplay = (DisplayRange.GetUpperBoundValue() - DisplayRange.GetLowerBoundValue()) * RenderData->GetSampleRate();
+	double FramesPixelDistance = InAllottedGeometry.Size.X / NumFramesToDisplay;
+	
+	TArray<FChannelSlotBoundaries> ChannelSlotsBoundaries;
+
+	for (uint16 Channel = 0; Channel < NChannels; ++Channel)
+	{
+		ChannelSlotsBoundaries.Emplace(Channel, NChannels, InAllottedGeometry);
+	}
+
+	for (double FrameX = GridMetrics.FirstMajorTickX; FrameX < InAllottedGeometry.Size.X + FramesPixelDistance; FrameX += FramesPixelDistance)
+	{
+		const double XRatio = FrameX / InAllottedGeometry.Size.X;
+		const double FrameTime = DisplayRange.GetLowerBoundValue() + DisplayRange.Size<double>() * XRatio;
+		const uint32 FrameIndex = FMath::RoundToInt(FrameTime * RenderData->GetSampleRate());
+		
+
+		if (FrameIndex >= 0 && FrameIndex < TotalWaveformFrames)
+		{
+			for (uint16 Channel = 0; Channel < NChannels; ++Channel)
+			{
+				const FChannelSlotBoundaries& ChannelBoundaries = ChannelSlotsBoundaries[Channel];
+				const uint32 SampleIndex = FrameIndex * NChannels + Channel;
+				const int16 SampleValue = RenderData->GetSampleData()[SampleIndex];
+
+				const float SampleValueRatio = (float)SampleValue / TNumericLimits<int16>::Max() * MaxWaveformHeight * VerticalZoomFactor;
+				const float TopBoundary = ChannelBoundaries.Top + MaxDistanceFromChannelBoundaries;
+				const float BottomBoundary = ChannelBoundaries.Bottom - MaxDistanceFromChannelBoundaries;
+				const float SampleY = FMath::Clamp((SampleValueRatio * ChannelBoundaries.Height / 2.f) + ChannelBoundaries.Center, TopBoundary, BottomBoundary);
+
+				const FVector2D SampleCoordinates(FrameX, SampleY);
+
+				OutDrawCoordinates.Add(SampleCoordinates);
+			}
+		}
+	}
+}
+
 void SWaveformViewer::DrawGridLines(const FGeometry& AllottedGeometry, FSlateWindowElementList& OutDrawElements, int32& LayerId) const
 {
 	TArray<FVector2D> LinePoints;
@@ -197,6 +311,29 @@ void SWaveformViewer::DrawGridLines(const FGeometry& AllottedGeometry, FSlateWin
 
 		}
 	}
+
+	const uint16 NChannels = RenderData->GetNumChannels();
+
+	for (uint16 Channel = 0; Channel < NChannels; ++Channel)
+	{
+		const FChannelSlotBoundaries ChannelBoundaries(Channel, NChannels, AllottedGeometry);
+
+		LinePoints[0] = FVector2D(0.f, ChannelBoundaries.Center);
+		LinePoints[1] = FVector2D(AllottedGeometry.Size.X, ChannelBoundaries.Center);
+
+		FSlateDrawElement::MakeLines(
+			OutDrawElements,
+			++LayerId,
+			AllottedGeometry.ToPaintGeometry(),
+			LinePoints,
+			ESlateDrawEffect::None,
+			ZeroCrossingLineColor.GetSpecifiedColor(),
+			false,
+			ZeroCrossingLineThickness);
+
+	}
+
+
 }
 
 void SWaveformViewer::UpdateGridMetrics(const FWaveEditorGridMetrics& InMetrics)
@@ -217,8 +354,12 @@ void SWaveformViewer::OnStyleUpdated(const FWaveformEditorWidgetStyleBase* Updat
 	WaveformColor = Style->WaveformColor;
 	MajorGridLineColor = Style->MajorGridLineColor;
 	MinorGridLineColor = Style->MinorGridLineColor;
+	ZeroCrossingLineColor = Style->ZeroCrossingLineColor;
+	ZeroCrossingLineThickness = Style->ZeroCrossingLineThickness;
 	BackgroundColor = Style->WaveformBackgroundColor;
 	BackgroundBrush = Style->BackgroundBrush;
+	SampleMarkersSize = Style->SampleMarkersSize;
+	WaveformLineThickness = Style->WaveformLineThickness;
 	DesiredWidth = Style->DesiredWidth;
 	DesiredHeight = Style->DesiredHeight;
 }
diff --git a/Engine/Plugins/Editor/WaveformEditor/Source/WaveformEditorWidgets/Private/SWaveformViewer.h b/Engine/Plugins/Editor/WaveformEditor/Source/WaveformEditorWidgets/Private/SWaveformViewer.h
--- a/Engine/Plugins/Editor/WaveformEditor/Source/WaveformEditorWidgets/Private/SWaveformViewer.h
+++ b/Engine/Plugins/Editor/WaveformEditor/Source/WaveformEditorWidgets/Private/SWaveformViewer.h
@@ -11,13 +11,33 @@
 #include "Widgets/SLeafWidget.h"
 
 class FWaveformEditorRenderData;
+struct FWaveformEditorWidgetStyleBase;
+struct FWaveformViewerStyle;
 
-struct FWaveformLineCoordinates
+struct FSamplesBinCoordinates
 {
 	FVector2D PointA;
 	FVector2D PointB;
 };
 
+struct FChannelSlotBoundaries
+{
+	explicit FChannelSlotBoundaries(const uint16 ChannelToDraw, const uint16 TotalNumChannels, const FGeometry& InAllottedGeometry)
+	{
+		Height = InAllottedGeometry.GetAbsoluteSize().Y / TotalNumChannels;
+		float ChannelSlotMidPoint = Height / 2.f;
+
+		Top = Height * ChannelToDraw;
+		Center = Top + ChannelSlotMidPoint;
+		Bottom = Top + Height;
+	}
+
+	float Top;
+	float Center;
+	float Bottom;
+	float Height;
+};
+
 class SWaveformViewer : public SLeafWidget
 {
 public:
@@ -35,14 +55,15 @@ public:
 	void OnRenderDataUpdated();
 	void OnDisplayRangeUpdated(const TRange<float> NewDisplayRange);
 	void UpdateGridMetrics(const FWaveEditorGridMetrics& InMetrics);
-
 	void OnStyleUpdated(const FWaveformEditorWidgetStyleBase* UpdatedStyle);
+
 private:
 	virtual int32 OnPaint(const FPaintArgs& Args, const FGeometry& AllottedGeometry, const FSlateRect& MyCullingRect, FSlateWindowElementList& OutDrawElements, int32 LayerId, const FWidgetStyle& InWidgetStyle, bool bParentEnabled) const override;
 	virtual FVector2D ComputeDesiredSize(float) const override;
 	virtual void Tick(const FGeometry& AllottedGeometry, const double InCurrentTime, const float InDeltaTime) override;
 
-	void GenerateWaveformLines(TArray<FWaveformLineCoordinates>& OutDrawCoordinates, const TArray<WaveformDrawingUtils::SampleRange>& InWaveformPeaks, const FGeometry& InAllottedGeometry, const float VerticalZoomFactor = 0.9f);
+	void GenerateWaveformBinsCoordinates(TArray<FSamplesBinCoordinates>& OutDrawCoordinates, const TArray<WaveformDrawingUtils::SampleRange>& InWaveformPeaks, const FGeometry& InAllottedGeometry, const float VerticalZoomFactor = 1.f);
+	void GenerateWaveformSamplesCoordinates(TArray<FVector2D>& OutDrawCoordinates, const FGeometry& InAllottedGeometry, const float VerticalZoomFactor = 1.f);
 	void DrawGridLines(const FGeometry& AllottedGeometry, FSlateWindowElementList& OutDrawElements, int32& LayerId) const;
 
 	TRange<float> DisplayRange;
@@ -58,15 +79,28 @@ private:
 	FSlateColor MajorGridLineColor = FLinearColor::Black;
 	FSlateColor MinorGridLineColor = FLinearColor(0.f, 0.f, 0.f, 0.5f);
 	FSlateColor ZeroCrossingLineColor = FLinearColor::Black;
+	float ZeroCrossingLineThickness = 1.f;
 	float DesiredHeight = 0.f;
 	float DesiredWidth = 0.f;
+	float SampleMarkersSize = 2.5f;
+	float WaveformLineThickness = 1.f;
+	
+	const float MaxDistanceFromChannelBoundaries = 2.f;
+	const float MaxWaveformHeight = 0.9f;
+	const float MinWaveformBinHeight = 0.1f;
+	const float MinScaledSamplePeakValue = 0.001f;
 
 	uint32 CachedPixelWidth = 0; 
 	float CachedPixelHeight = 0.f;
 	TArray<WaveformDrawingUtils::SampleRange> CachedPeaks;
-	TArray<FWaveformLineCoordinates> CachedDrawCoordinates;
-
+	TArray<FSamplesBinCoordinates> CachedBinsDrawCoordinates;
+	TArray<FVector2D> CachedSampleDrawCoordinates;
 
 	TSharedPtr<FWaveformEditorTransportCoordinator> TransportCoordinator = nullptr;
 	TSharedPtr<FWaveformEditorRenderData> RenderData = nullptr;
+
+	enum class EWaveformDrawMode {
+		BinnedPeaks = 0,
+		WaveformLine
+	} WaveformDrawMode;
 };
\ No newline at end of file
diff --git a/Engine/Plugins/Editor/WaveformEditor/Source/WaveformEditorWidgets/Private/WaveformEditorSlateTypes.cpp b/Engine/Plugins/Editor/WaveformEditor/Source/WaveformEditorWidgets/Private/WaveformEditorSlateTypes.cpp
--- a/Engine/Plugins/Editor/WaveformEditor/Source/WaveformEditorWidgets/Private/WaveformEditorSlateTypes.cpp
+++ b/Engine/Plugins/Editor/WaveformEditor/Source/WaveformEditorWidgets/Private/WaveformEditorSlateTypes.cpp
@@ -20,8 +20,12 @@ const FName FWaveformViewerStyle::TypeName("FWaveformViewerStyle");
 
 FWaveformViewerStyle::FWaveformViewerStyle()
 	: WaveformColor(FLinearColor::White)
+	, WaveformLineThickness(1.f)
 	, MajorGridLineColor(FLinearColor::Black)
 	, MinorGridLineColor(FLinearColor(0.f, 0.f, 0.f, 0.5f))
+	, ZeroCrossingLineColor(FLinearColor::Black)
+	, ZeroCrossingLineThickness(1.f)
+	, SampleMarkersSize(2.5f)
 	, WaveformBackgroundColor(WaveformEditorStylesSharedParams::WaveformBackGroundColor)
 	, BackgroundBrush(*FAppStyle::GetBrush(WaveformEditorStylesSharedParams::BackgroundBrushName))
 	, DesiredWidth(WaveformEditorStylesSharedParams::ViewerWidth)
diff --git a/Engine/Plugins/Editor/WaveformEditor/Source/WaveformEditorWidgets/Private/WaveformEditorStyle.cpp b/Engine/Plugins/Editor/WaveformEditor/Source/WaveformEditorWidgets/Private/WaveformEditorStyle.cpp
--- a/Engine/Plugins/Editor/WaveformEditor/Source/WaveformEditorWidgets/Private/WaveformEditorStyle.cpp
+++ b/Engine/Plugins/Editor/WaveformEditor/Source/WaveformEditorWidgets/Private/WaveformEditorStyle.cpp
@@ -48,8 +48,12 @@ void FWaveformEditorStyle::Init()
 	FWaveformViewerStyle WaveViewerStyle = FWaveformViewerStyle()
 		.SetWaveformColor(Settings->WaveformColor)
 		.SetBackgroundColor(Settings->WaveformBackgroundColor)
+		.SetWaveformLineThickness(Settings->WaveformLineThickness)
+		.SetSampleMarkersSize(Settings->SampleMarkersSize)
 		.SetMajorGridLineColor(Settings->MajorGridColor)
-		.SetMinorGridLineColor(Settings->MinorGridColor);
+		.SetMinorGridLineColor(Settings->MinorGridColor)
+		.SetZeroCrossingLineColor(Settings->ZeroCrossingLineColor)
+		.SetZeroCrossingLineThickness(Settings->ZeroCrossingLineThickness);
 
 	StyleInstance->Set(WaveformViewerStyleName, WaveViewerStyle);
 
@@ -93,6 +97,14 @@ void FWaveformEditorStyle::OnWidgetSettingsUpdated(const FName& PropertyName, co
 	{
 		WaveformViewerStyle->SetBackgroundColor(Settings->WaveformBackgroundColor);
 	}
+	else if (PropertyName == GET_MEMBER_NAME_CHECKED(UWaveformEditorWidgetsSettings, WaveformLineThickness))
+	{
+		WaveformViewerStyle->SetWaveformLineThickness(Settings->WaveformLineThickness);
+	}
+	else if (PropertyName == GET_MEMBER_NAME_CHECKED(UWaveformEditorWidgetsSettings, SampleMarkersSize))
+	{
+		WaveformViewerStyle->SetSampleMarkersSize(Settings->SampleMarkersSize);
+	}
 	else if (PropertyName == GET_MEMBER_NAME_CHECKED(UWaveformEditorWidgetsSettings, MajorGridColor))
 	{
 		WaveformViewerStyle->SetMajorGridLineColor(Settings->MajorGridColor);
@@ -101,6 +113,14 @@ void FWaveformEditorStyle::OnWidgetSettingsUpdated(const FName& PropertyName, co
 	{
 		WaveformViewerStyle->SetMinorGridLineColor(Settings->MinorGridColor);
 	}
+	else if (PropertyName == GET_MEMBER_NAME_CHECKED(UWaveformEditorWidgetsSettings, ZeroCrossingLineColor))
+	{
+		WaveformViewerStyle->SetZeroCrossingLineColor(Settings->ZeroCrossingLineColor);
+	}
+	else if (PropertyName == GET_MEMBER_NAME_CHECKED(UWaveformEditorWidgetsSettings, ZeroCrossingLineThickness))
+	{
+		WaveformViewerStyle->SetZeroCrossingLineThickness(Settings->ZeroCrossingLineThickness);
+	}
 	else if (PropertyName == GET_MEMBER_NAME_CHECKED(UWaveformEditorWidgetsSettings, PlayheadColor))
 	{
 		WaveformViewerOverlayStyle->SetPlayheadColor(Settings->PlayheadColor);
@@ -122,4 +142,4 @@ void FWaveformEditorStyle::OnWidgetSettingsUpdated(const FName& PropertyName, co
 	{
 		WaveformEditorTimeRulerStyle->SetFontSize(Settings->RulerFontSize);
 	}
-}
\ No newline at end of file
+}
diff --git a/Engine/Plugins/Editor/WaveformEditor/Source/WaveformEditorWidgets/Private/WaveformEditorWidgetsSettings.cpp b/Engine/Plugins/Editor/WaveformEditor/Source/WaveformEditorWidgets/Private/WaveformEditorWidgetsSettings.cpp
--- a/Engine/Plugins/Editor/WaveformEditor/Source/WaveformEditorWidgets/Private/WaveformEditorWidgetsSettings.cpp
+++ b/Engine/Plugins/Editor/WaveformEditor/Source/WaveformEditorWidgets/Private/WaveformEditorWidgetsSettings.cpp
@@ -6,7 +6,11 @@ UWaveformEditorWidgetsSettings::UWaveformEditorWidgetsSettings(const FObjectInit
 	: Super(ObjectInitializer)
 	, PlayheadColor(FLinearColor(255.f, 0.1f, 0.2f, 1.f))
 	, WaveformColor(FLinearColor::White)
+	, WaveformLineThickness(1.f)
+	, SampleMarkersSize(2.5f)
 	, WaveformBackgroundColor(FLinearColor(0.02f, 0.02f, 0.02f, 1.f))
+	, ZeroCrossingLineColor(FLinearColor::Black)
+	, ZeroCrossingLineThickness(1.f)
 	, MajorGridColor(FLinearColor::Black)
 	, MinorGridColor(FLinearColor(0.f, 0.f, 0.f, 0.5f))
 	, RulerBackgroundColor(FLinearColor::Black)
diff --git a/Engine/Plugins/Editor/WaveformEditor/Source/WaveformEditorWidgets/Private/WaveformEditorWidgetsSettings.h b/Engine/Plugins/Editor/WaveformEditor/Source/WaveformEditorWidgets/Private/WaveformEditorWidgetsSettings.h
--- a/Engine/Plugins/Editor/WaveformEditor/Source/WaveformEditorWidgets/Private/WaveformEditorWidgetsSettings.h
+++ b/Engine/Plugins/Editor/WaveformEditor/Source/WaveformEditorWidgets/Private/WaveformEditorWidgetsSettings.h
@@ -38,9 +38,21 @@ public:
 	UPROPERTY(config, EditAnywhere, Category = "Waveform Viewer")
 	FLinearColor WaveformColor;
 
+	UPROPERTY(config, EditAnywhere, Category = "Waveform Viewer", Meta = (ClampMin = "1", ClampMax = "10"))
+	float WaveformLineThickness;
+
+	UPROPERTY(config, EditAnywhere, Category = "Waveform Viewer", Meta = (ClampMin = "0", ClampMax = "30"))
+	float SampleMarkersSize;
+
 	UPROPERTY(config, EditAnywhere, Category = "Waveform Viewer")
 	FLinearColor WaveformBackgroundColor;
 
+	UPROPERTY(config, EditAnywhere, Category = "Waveform Viewer")
+	FLinearColor ZeroCrossingLineColor;
+
+	UPROPERTY(config, EditAnywhere, Category = "Waveform Viewer", Meta = (ClampMin = "1", ClampMax = "10"))
+	float ZeroCrossingLineThickness;
+
 	UPROPERTY(config, EditAnywhere, Category = "Waveform Viewer")
 	FLinearColor MajorGridColor;
 	
diff --git a/Engine/Plugins/Editor/WaveformEditor/Source/WaveformEditorWidgets/Public/WaveformEditorSlateTypes.h b/Engine/Plugins/Editor/WaveformEditor/Source/WaveformEditorWidgets/Public/WaveformEditorSlateTypes.h
--- a/Engine/Plugins/Editor/WaveformEditor/Source/WaveformEditorWidgets/Public/WaveformEditorSlateTypes.h
+++ b/Engine/Plugins/Editor/WaveformEditor/Source/WaveformEditorWidgets/Public/WaveformEditorSlateTypes.h
@@ -43,6 +43,10 @@ struct FWaveformViewerStyle : public FWaveformEditorWidgetStyleBase
 	FSlateColor WaveformColor;
 	FWaveformViewerStyle& SetWaveformColor(const FSlateColor InWaveformColor) { WaveformColor = InWaveformColor; BroadcastStyleUpdate(); return *this; }
 
+	UPROPERTY(EditAnywhere, BlueprintReadWrite, Category = Appearance)
+	float WaveformLineThickness;
+	FWaveformViewerStyle& SetWaveformLineThickness(const float InWaveformLineThickness) { WaveformLineThickness = InWaveformLineThickness; BroadcastStyleUpdate(); return *this; }
+
 	UPROPERTY(EditAnywhere, BlueprintReadWrite, Category = Appearance)
 	FSlateColor MajorGridLineColor;
 	FWaveformViewerStyle& SetMajorGridLineColor(const FSlateColor InMajorGridLineColor) { MajorGridLineColor = InMajorGridLineColor; BroadcastStyleUpdate(); return *this; }
@@ -51,6 +55,18 @@ struct FWaveformViewerStyle : public FWaveformEditorWidgetStyleBase
 	FSlateColor MinorGridLineColor;
 	FWaveformViewerStyle& SetMinorGridLineColor(const FSlateColor InMinorGridLineColor) { MinorGridLineColor = InMinorGridLineColor; BroadcastStyleUpdate(); return *this; }
 
+	UPROPERTY(EditAnywhere, BlueprintReadWrite, Category = Appearance)
+	FSlateColor ZeroCrossingLineColor;
+	FWaveformViewerStyle& SetZeroCrossingLineColor(const FSlateColor InZeroCrossingLineColor) { ZeroCrossingLineColor = InZeroCrossingLineColor; BroadcastStyleUpdate(); return *this; }
+
+	UPROPERTY(EditAnywhere, BlueprintReadWrite, Category = Appearance)
+	float ZeroCrossingLineThickness;
+	FWaveformViewerStyle& SetZeroCrossingLineThickness(const float InZeroCrossingLineThickness) { ZeroCrossingLineThickness = InZeroCrossingLineThickness; BroadcastStyleUpdate(); return *this; }
+
+	UPROPERTY(EditAnywhere, BlueprintReadWrite, Category = Appearance)
+	float SampleMarkersSize;
+	FWaveformViewerStyle& SetSampleMarkersSize(const float InSampleMarkersSize) { SampleMarkersSize = InSampleMarkersSize; BroadcastStyleUpdate(); return *this; }
+
 	UPROPERTY(EditAnywhere, BlueprintReadWrite, Category = Appearance)
 	FSlateColor WaveformBackgroundColor;
 	FWaveformViewerStyle& SetBackgroundColor(const FSlateColor InBackgroundColor) { WaveformBackgroundColor = InBackgroundColor; BroadcastStyleUpdate(); return *this; }
