[Waveform Editor] Transformations Details view is now a custom USoundWave visualization

diff --git a/Engine/Plugins/Editor/WaveformEditor/Source/WaveformEditor/Private/WaveformEditor.cpp b/Engine/Plugins/Editor/WaveformEditor/Source/WaveformEditor/Private/WaveformEditor.cpp
--- a/Engine/Plugins/Editor/WaveformEditor/Source/WaveformEditor/Private/WaveformEditor.cpp
+++ b/Engine/Plugins/Editor/WaveformEditor/Source/WaveformEditor/Private/WaveformEditor.cpp
@@ -11,6 +11,7 @@
 #include "SWaveformTransformationsOverlay.h"
 #include "ToolMenus.h"
 #include "WaveformEditorCommands.h"
+#include "WaveformEditorDetailsCustomization.h"
 #include "WaveformEditorLog.h"
 #include "WaveformEditorRenderData.h"
 #include "WaveformEditorStyle.h"
@@ -489,9 +490,13 @@ bool FWaveformEditor::SetUpDetailsViews()
 	PropertiesDetails->SetObject(SoundWave);
 
 	TransformationsDetails = PropertyModule.CreateDetailView(Args);
-	TransformationsView = TStrongObjectPtr(NewObject<UWaveformTransformationsViewHelper>());
-	TransformationsView->SetSoundWave(SoundWave);
-	TransformationsDetails->SetObject(TransformationsView.Get());
+	FOnGetDetailCustomizationInstance TransformationsDetailsCustomizationInstance = FOnGetDetailCustomizationInstance::CreateLambda([]() { 
+			return MakeShared<FWaveformTransformationsDetailsCustomization>(); 
+		}
+	);
+
+	TransformationsDetails->RegisterInstancedCustomPropertyLayout(SoundWave->GetClass(), TransformationsDetailsCustomizationInstance);
+	TransformationsDetails->SetObject(SoundWave);
 
 	return true;
 }
diff --git a/Engine/Plugins/Editor/WaveformEditor/Source/WaveformEditor/Private/WaveformEditorCustomDetailsHelpers.cpp b/Engine/Plugins/Editor/WaveformEditor/Source/WaveformEditor/Private/WaveformEditorCustomDetailsHelpers.cpp
deleted file mode 100644
diff --git a/Engine/Plugins/Editor/WaveformEditor/Source/WaveformEditor/Private/WaveformEditorDetailsCustomization.cpp b/Engine/Plugins/Editor/WaveformEditor/Source/WaveformEditor/Private/WaveformEditorDetailsCustomization.cpp
--- a/Engine/Plugins/Editor/WaveformEditor/Source/WaveformEditor/Private/WaveformEditorDetailsCustomization.cpp
+++ b/Engine/Plugins/Editor/WaveformEditor/Source/WaveformEditor/Private/WaveformEditorDetailsCustomization.cpp
@@ -5,24 +5,23 @@
 #include "DetailCategoryBuilder.h"
 #include "DetailLayoutBuilder.h"
 #include "Sound/SoundWave.h"
-#include "WaveformEditorCustomDetailsHelpers.h"
+
+static const FLazyName TransformationsCategoryName("Transformations");
 
 void FWaveformTransformationsDetailsCustomization::CustomizeDetails(IDetailLayoutBuilder& DetailLayout)
 {	
-	TArray<TWeakObjectPtr<UObject>> Objects;
-	DetailLayout.GetObjectsBeingCustomized(Objects);
-
-	//there should only one soundwave per editor instance open
-	check(Objects.Num() == 1)
+	TArray<FName> CategoryNames;
+	DetailLayout.GetCategoryNames(CategoryNames);
 
-	if (UWaveformTransformationsViewHelper* TransformationsView = CastChecked<UWaveformTransformationsViewHelper>(Objects.Last()))
+	for (FName& CategoryName : CategoryNames)
 	{
-		SoundWaveObject = TransformationsView->GetSoundWave();
+		if (CategoryName != TransformationsCategoryName)
+		{
+			DetailLayout.HideCategory(CategoryName);
+		}
 	}
 
-	IDetailCategoryBuilder& CategoryBuilder = DetailLayout.EditCategory("General");
-
-	TArray<UObject*> ObjArray{ SoundWaveObject.Get() };
-	CategoryBuilder.AddExternalObjectProperty(ObjArray, GET_MEMBER_NAME_CHECKED(USoundWave, Transformations));
+	IDetailCategoryBuilder& CategoryBuilder = DetailLayout.EditCategory("Transformations");
 	CategoryBuilder.InitiallyCollapsed(false);
+	CategoryBuilder.RestoreExpansionState(true);
 }
\ No newline at end of file
diff --git a/Engine/Plugins/Editor/WaveformEditor/Source/WaveformEditor/Private/WaveformEditorModule.cpp b/Engine/Plugins/Editor/WaveformEditor/Source/WaveformEditor/Private/WaveformEditorModule.cpp
--- a/Engine/Plugins/Editor/WaveformEditor/Source/WaveformEditor/Private/WaveformEditorModule.cpp
+++ b/Engine/Plugins/Editor/WaveformEditor/Source/WaveformEditor/Private/WaveformEditorModule.cpp
@@ -5,8 +5,6 @@
 #include "Modules/ModuleManager.h"
 #include "PropertyEditorModule.h"
 #include "WaveformEditorCommands.h"
-#include "WaveformEditorCustomDetailsHelpers.h"
-#include "WaveformEditorDetailsCustomization.h"
 #include "WaveformEditorInstantiator.h"
 #include "WaveformEditorLog.h"
 
@@ -19,12 +17,6 @@ void FWaveformEditorModule::StartupModule()
 
 	WaveformEditorInstantiator = MakeShared<FWaveformEditorInstantiator>();
 	RegisterContentBrowserExtensions(WaveformEditorInstantiator.Get());
-
-	FPropertyEditorModule& PropertyModule = FModuleManager::LoadModuleChecked<FPropertyEditorModule>("PropertyEditor");
-
-	PropertyModule.RegisterCustomClassLayout(
-		UWaveformTransformationsViewHelper::StaticClass()->GetFName(),
-		FOnGetDetailCustomizationInstance::CreateLambda([]() { return MakeShared<FWaveformTransformationsDetailsCustomization>(); }));
 }
 
 void FWaveformEditorModule::ShutdownModule()
diff --git a/Engine/Plugins/Editor/WaveformEditor/Source/WaveformEditor/Public/WaveformEditor.h b/Engine/Plugins/Editor/WaveformEditor/Source/WaveformEditor/Public/WaveformEditor.h
--- a/Engine/Plugins/Editor/WaveformEditor/Source/WaveformEditor/Public/WaveformEditor.h
+++ b/Engine/Plugins/Editor/WaveformEditor/Source/WaveformEditor/Public/WaveformEditor.h
@@ -7,7 +7,6 @@
 #include "Toolkits/AssetEditorToolkit.h"
 #include "WaveformEditorTransportController.h"
 #include "WaveformEditorTransportCoordinator.h"
-#include "WaveformEditorCustomDetailsHelpers.h"
 #include "WaveformEditorZoomController.h"
 
 class IToolkitHost;
@@ -112,7 +111,6 @@ private:
 
 	/** Transformations tab */
 	TSharedPtr<IDetailsView> TransformationsDetails;
-	TStrongObjectPtr<UWaveformTransformationsViewHelper> TransformationsView;
 
 	/** Settings Editor App Identifier */
 	static const FName AppIdentifier;
diff --git a/Engine/Plugins/Editor/WaveformEditor/Source/WaveformEditor/Public/WaveformEditorCustomDetailsHelpers.h b/Engine/Plugins/Editor/WaveformEditor/Source/WaveformEditor/Public/WaveformEditorCustomDetailsHelpers.h
deleted file mode 100644
