[WaveformEditor] Fix for undo not working from transformations pane
#trivial

diff --git a/Engine/Plugins/Editor/WaveformEditor/Source/WaveformEditor/Private/WaveformEditor.cpp b/Engine/Plugins/Editor/WaveformEditor/Source/WaveformEditor/Private/WaveformEditor.cpp
--- a/Engine/Plugins/Editor/WaveformEditor/Source/WaveformEditor/Private/WaveformEditor.cpp
+++ b/Engine/Plugins/Editor/WaveformEditor/Source/WaveformEditor/Private/WaveformEditor.cpp
@@ -455,7 +455,9 @@ bool FWaveformEditor::MatchesContext(const FTransactionContext& InContext, const
 	{
 		UObject* Object = TransactionObjectPair.Key;
 		const UClass* ObjectClass = Object->GetClass();
-		if (ObjectClass && ObjectClass->IsChildOf(UWaveformTransformationBase::StaticClass()))
+		const bool bIsSupportedUndoClass = ObjectClass->IsChildOf(UWaveformTransformationBase::StaticClass()) || ObjectClass->IsChildOf(USoundWave::StaticClass());
+
+		if (ObjectClass && bIsSupportedUndoClass)
 		{
 			bShoouldUndo = true;
 		}
