[Waveform Editor] Fix playhead ignoring play position set on pause
fix for #jira UE-164880
#trivial

diff --git a/Engine/Plugins/Editor/WaveformEditor/Source/WaveformEditor/Private/WaveformEditorTransportController.cpp b/Engine/Plugins/Editor/WaveformEditor/Source/WaveformEditor/Private/WaveformEditorTransportController.cpp
--- a/Engine/Plugins/Editor/WaveformEditor/Source/WaveformEditor/Private/WaveformEditorTransportController.cpp
+++ b/Engine/Plugins/Editor/WaveformEditor/Source/WaveformEditor/Private/WaveformEditorTransportController.cpp
@@ -21,11 +21,16 @@ void FWaveformEditorTransportController::Play()
 	if (IsPaused())
 	{
 		AudioComponent->SetPaused(false);
-		return;
+
+		if (!bCachedTimeDuringPause)
+		{
+			return;
+		}
+		
 	}
 
 	AudioComponent->Play(CachedAudioStartTime);
-
+	
 }
 
 void FWaveformEditorTransportController::Play(const float StartTime)
@@ -64,6 +69,8 @@ void FWaveformEditorTransportController::Stop()
 	{
 		AudioComponent->SetPaused(false);
 	}
+
+	bCachedTimeDuringPause = false;
 }
 
 void FWaveformEditorTransportController::TogglePlayback()
@@ -101,6 +108,15 @@ bool FWaveformEditorTransportController::IsPlaying() const
 void FWaveformEditorTransportController::CacheStartTime(const float StartTime)
 {
 	CachedAudioStartTime = StartTime;
+
+	if (IsPaused())
+	{
+		bCachedTimeDuringPause = true;
+	}
+	else
+	{
+		bCachedTimeDuringPause = false;
+	}
 }
 
 void FWaveformEditorTransportController::Seek(const float SeekTime)
@@ -111,4 +127,4 @@ void FWaveformEditorTransportController::Seek(const float SeekTime)
 const bool FWaveformEditorTransportController::SoundBaseIsValid() const
 {
 	return AudioComponent->GetSound() != nullptr;
-}
+}
\ No newline at end of file
diff --git a/Engine/Plugins/Editor/WaveformEditor/Source/WaveformEditor/Public/WaveformEditorTransportController.h b/Engine/Plugins/Editor/WaveformEditor/Source/WaveformEditor/Public/WaveformEditorTransportController.h
--- a/Engine/Plugins/Editor/WaveformEditor/Source/WaveformEditor/Public/WaveformEditorTransportController.h
+++ b/Engine/Plugins/Editor/WaveformEditor/Source/WaveformEditor/Public/WaveformEditorTransportController.h
@@ -27,6 +27,7 @@ private:
 
 	UAudioComponent* AudioComponent = nullptr;
 	float CachedAudioStartTime = 0.f;
+	bool bCachedTimeDuringPause = false;
 
 	FDelegateHandle PlayStateChangeDelegateHandle;
 };
\ No newline at end of file
