[WaveformEditor] Fix for warning V595 in FWaveformEditor::MatchesContext
#trivial

diff --git a/Engine/Plugins/Editor/WaveformEditor/Source/WaveformEditor/Private/WaveformEditor.cpp b/Engine/Plugins/Editor/WaveformEditor/Source/WaveformEditor/Private/WaveformEditor.cpp
--- a/Engine/Plugins/Editor/WaveformEditor/Source/WaveformEditor/Private/WaveformEditor.cpp
+++ b/Engine/Plugins/Editor/WaveformEditor/Source/WaveformEditor/Private/WaveformEditor.cpp
@@ -455,11 +455,10 @@ bool FWaveformEditor::MatchesContext(const FTransactionContext& InContext, const
 	{
 		UObject* Object = TransactionObjectPair.Key;
 		const UClass* ObjectClass = Object->GetClass();
-		const bool bIsSupportedUndoClass = ObjectClass->IsChildOf(UWaveformTransformationBase::StaticClass()) || ObjectClass->IsChildOf(USoundWave::StaticClass());
 
-		if (ObjectClass && bIsSupportedUndoClass)
+		if (ObjectClass)
 		{
-			bShoouldUndo = true;
+			bShoouldUndo = ObjectClass->IsChildOf(UWaveformTransformationBase::StaticClass()) || ObjectClass->IsChildOf(USoundWave::StaticClass());
 		}
 	}
 
