[WaveformTransformations] Add Configuration handling to UWaveformTransformationBase  for parameter range checks ecc
fix for #jira UE-163720

diff --git a/Engine/Plugins/Editor/WaveformEditor/Source/WaveformEditor/Private/WaveformEditor.cpp b/Engine/Plugins/Editor/WaveformEditor/Source/WaveformEditor/Private/WaveformEditor.cpp
--- a/Engine/Plugins/Editor/WaveformEditor/Source/WaveformEditor/Private/WaveformEditor.cpp
+++ b/Engine/Plugins/Editor/WaveformEditor/Source/WaveformEditor/Private/WaveformEditor.cpp
@@ -104,6 +104,21 @@ void FWaveformEditor::AddDefaultTransformations()
 				SoundWave->Transformations.Add(nullptr);
 			}
 		}
+
+		FWaveTransformUObjectConfiguration ChainMetadata = FWaveTransformUObjectConfiguration();
+
+		ChainMetadata.NumChannels = SoundWave->NumChannels;
+		ChainMetadata.SampleRate = SoundWave->GetSampleRateForCurrentPlatform();
+		ChainMetadata.StartTime = 0.f;
+		ChainMetadata.EndTime = SoundWave->Duration;
+
+		for (TObjectPtr<UWaveformTransformationBase>& Transform : SoundWave->Transformations)
+		{
+			if (Transform)
+			{
+				Transform->UpdateConfiguration(ChainMetadata);
+			}
+		}
 	}
 }
 
diff --git a/Engine/Plugins/Editor/WaveformEditor/Source/WaveformTransformations/Private/WaveformTransformationTrimFade.cpp b/Engine/Plugins/Editor/WaveformEditor/Source/WaveformTransformations/Private/WaveformTransformationTrimFade.cpp
--- a/Engine/Plugins/Editor/WaveformEditor/Source/WaveformTransformations/Private/WaveformTransformationTrimFade.cpp
+++ b/Engine/Plugins/Editor/WaveformEditor/Source/WaveformTransformations/Private/WaveformTransformationTrimFade.cpp
@@ -135,56 +135,18 @@ Audio::FTransformationPtr UWaveformTransformationTrimFade::CreateTransformation(
 	return MakeUnique<FWaveTransformationTrimFade>(StartTime, EndTime, StartFadeTime, StartFadeCurve, EndFadeTime, EndFadeCurve);
 }
 
-void UWaveformTransformationTrimFade::PostEditChangeProperty(struct FPropertyChangedEvent& PropertyChangedEvent)
+void UWaveformTransformationTrimFade::UpdateConfiguration(FWaveTransformUObjectConfiguration& InOutConfiguration)
 {
-	Super::PostEditChangeProperty(PropertyChangedEvent);
-	
-	if (PropertyChangedEvent.GetPropertyName() == WaveformTransformationTrimFadeNames::StartTimeName 
-		|| PropertyChangedEvent.GetPropertyName() == WaveformTransformationTrimFadeNames::EndTimeName)
-	{
-		if (AvailableWaveformDuration < 0.f)
-		{
-			UpdateAvailableWaveformDuration();
-		}
-
-		StartTime = FMath::Clamp(StartTime, 0.f, AvailableWaveformDuration - UE_KINDA_SMALL_NUMBER);
-		EndTime = FMath::Clamp(EndTime, StartTime + UE_KINDA_SMALL_NUMBER, AvailableWaveformDuration);
-	}
-}
+	UpdateDurationProperties(InOutConfiguration.EndTime - InOutConfiguration.StartTime);
 
-void UWaveformTransformationTrimFade::PostInitProperties()
-{
-	Super::PostInitProperties();
-	
-	if (EndTime < 0.f)
-	{
-		UpdateAvailableWaveformDuration();
-		EndTime = AvailableWaveformDuration;
-	}
+	InOutConfiguration.StartTime = StartTime;
+	InOutConfiguration.EndTime = EndTime;
 }
 
-void UWaveformTransformationTrimFade::UpdateAvailableWaveformDuration()
+void UWaveformTransformationTrimFade::UpdateDurationProperties(const float InAvailableDuration)
 {
-	const USoundWave* ParentSoundWave = this->GetTypedOuter<USoundWave>();
-
-	if (!ParentSoundWave)
-	{
-		return;
-	}
-
-	AvailableWaveformDuration = ParentSoundWave->Duration;
-	
-	for (TObjectPtr<UWaveformTransformationBase> Transformation : ParentSoundWave->Transformations)
-	{
-		if (Transformation && Transformation->IsA(UWaveformTransformationTrimFade::StaticClass()))
-		{
-			if (Transformation == this)
-			{
-				return;
-			}
-
-			TObjectPtr<UWaveformTransformationTrimFade> TrimFadeTransformPtr = Cast<UWaveformTransformationTrimFade>(Transformation);
-			AvailableWaveformDuration = TrimFadeTransformPtr->EndTime - TrimFadeTransformPtr->StartTime;
-		}
-	}
-}
+	check(InAvailableDuration > 0)
+	AvailableWaveformDuration = InAvailableDuration;
+	StartTime = FMath::Clamp(StartTime, 0.f, AvailableWaveformDuration - UE_KINDA_SMALL_NUMBER);
+	EndTime = EndTime < 0 ? EndTime = AvailableWaveformDuration : FMath::Clamp(EndTime, StartTime + UE_KINDA_SMALL_NUMBER, AvailableWaveformDuration);
+}
\ No newline at end of file
diff --git a/Engine/Plugins/Editor/WaveformEditor/Source/WaveformTransformations/Public/WaveformTransformationTrimFade.h b/Engine/Plugins/Editor/WaveformEditor/Source/WaveformTransformations/Public/WaveformTransformationTrimFade.h
--- a/Engine/Plugins/Editor/WaveformEditor/Source/WaveformTransformations/Public/WaveformTransformationTrimFade.h
+++ b/Engine/Plugins/Editor/WaveformEditor/Source/WaveformTransformations/Public/WaveformTransformationTrimFade.h
@@ -52,15 +52,11 @@ public:
 
 	virtual Audio::FTransformationPtr CreateTransformation() const override;
 
-	void PostEditChangeProperty(struct FPropertyChangedEvent& PropertyChangedEvent) override;
-	void PostInitProperties() override;
-
+	void UpdateConfiguration(FWaveTransformUObjectConfiguration& InOutConfiguration) override;
 	
-
-private: 
-	void UpdateAvailableWaveformDuration();
-	float AvailableWaveformDuration = - 1.f;
-
+private:	
+	void UpdateDurationProperties(const float InAvailableDuration);
+	float AvailableWaveformDuration = -1.f;
 };
 
 // class UBaseClass;
diff --git a/Engine/Plugins/Editor/WaveformEditor/Source/WaveformTransformationsWidgets/Private/WaveformTransformationTrimFadeRenderer.cpp b/Engine/Plugins/Editor/WaveformEditor/Source/WaveformTransformationsWidgets/Private/WaveformTransformationTrimFadeRenderer.cpp
--- a/Engine/Plugins/Editor/WaveformEditor/Source/WaveformTransformationsWidgets/Private/WaveformTransformationTrimFadeRenderer.cpp
+++ b/Engine/Plugins/Editor/WaveformEditor/Source/WaveformTransformationsWidgets/Private/WaveformTransformationTrimFadeRenderer.cpp
@@ -197,7 +197,7 @@ void FWaveformTransformationTrimFadeRenderer::Tick(const FGeometry& AllottedGeom
 	const double FirstFrame = FMath::Clamp((TrimFadeTransform->StartTime * TransformationWaveInfo.SampleRate) , 0.f, NumFrames);
 	const double EndFrame = FMath::Clamp((TrimFadeTransform->EndTime * TransformationWaveInfo.SampleRate), FirstFrame, NumFrames);
 
-	PixelsPerFrame = AllottedGeometry.GetLocalSize().X / NumFrames;
+	PixelsPerFrame = FMath::Max(AllottedGeometry.GetLocalSize().X / NumFrames, 0.0);
 
 	StartTimeHandleX = FirstFrame * PixelsPerFrame;
 	EndTimeHandleX = EndFrame * PixelsPerFrame;
diff --git a/Engine/Source/Runtime/AudioExtensions/Public/IWaveformTransformation.h b/Engine/Source/Runtime/AudioExtensions/Public/IWaveformTransformation.h
--- a/Engine/Source/Runtime/AudioExtensions/Public/IWaveformTransformation.h
+++ b/Engine/Source/Runtime/AudioExtensions/Public/IWaveformTransformation.h
@@ -49,6 +49,15 @@ namespace Audio
 	using FTransformationPtr = TUniquePtr<Audio::IWaveTransformation>;
 }
 
+// Information about the the wave file we are transforming for Transformation UObjects
+struct FWaveTransformUObjectConfiguration
+{
+	int32 NumChannels = 0;
+	float SampleRate = 0;
+	float StartTime = 0.f; 
+	float EndTime = -1.f; 
+};
+
 // Base class to hold editor configurable properties for an arbitrary transformation of audio waveform data
 UCLASS(Abstract, EditInlineNew)
 class AUDIOEXTENSIONS_API UWaveformTransformationBase : public UObject
@@ -57,6 +66,7 @@ class AUDIOEXTENSIONS_API UWaveformTransformationBase : public UObject
 
 public:
 	virtual Audio::FTransformationPtr CreateTransformation() const { return nullptr; }
+	virtual void UpdateConfiguration(FWaveTransformUObjectConfiguration& InOutConfiguration) {};
 };
 
 // Object that holds an ordered list of transformations to perform on a sound wave
diff --git a/Engine/Source/Runtime/Engine/Private/SoundWave.cpp b/Engine/Source/Runtime/Engine/Private/SoundWave.cpp
--- a/Engine/Source/Runtime/Engine/Private/SoundWave.cpp
+++ b/Engine/Source/Runtime/Engine/Private/SoundWave.cpp
@@ -2478,14 +2478,33 @@ void USoundWave::PostEditChangeProperty(FPropertyChangedEvent& PropertyChangedEv
 	SoundWaveDataPtr->bIsStreaming = IsStreaming(nullptr);
 	SoundWaveDataPtr->bShouldUseStreamCaching = ShouldUseStreamCaching();
 
-	// Prevent constant re-compression of SoundWave while properties are being changed interactively
-	if (PropertyChangedEvent.ChangeType != EPropertyChangeType::Interactive)
+	if (FProperty* PropertyThatChanged = PropertyChangedEvent.Property)
 	{
-		// Regenerate on save any compressed sound formats or if analysis needs to be re-done
-		if (FProperty* PropertyThatChanged = PropertyChangedEvent.Property)
+		const FName& Name = PropertyThatChanged->GetFName();
+
+		//Keep transformations UObjects configurations up to date
+		if (Name == TransformationsFName)
 		{
-			const FName& Name = PropertyThatChanged->GetFName();
+			FWaveTransformUObjectConfiguration ChainConfiguration = FWaveTransformUObjectConfiguration();
+
+			ChainConfiguration.NumChannels = NumChannels;
+			ChainConfiguration.SampleRate = GetSampleRateForCurrentPlatform();
+			ChainConfiguration.StartTime = 0.f;
+			ChainConfiguration.EndTime = Duration;
 
+			for (TObjectPtr<UWaveformTransformationBase>& Transformation : Transformations)
+			{
+				if (Transformation)
+				{
+					Transformation->UpdateConfiguration(ChainConfiguration);
+				}
+			}
+		}
+	
+		// Prevent constant re-compression of SoundWave while properties are being changed interactively
+		if (PropertyChangedEvent.ChangeType != EPropertyChangeType::Interactive)
+		{
+			// Regenerate on save any compressed sound formats or if analysis needs to be re-done
 			if (Name == LoadingBehaviorFName)
 			{
 				// Update and cache new loading behavior if it has changed. This
