[WaveformTransformationNormalize] Convert RMS peak from linear to decibel WaveformTransformationNormalize
fix for #jira UE-165307
#trivial

diff --git a/Engine/Plugins/Editor/WaveformEditor/Source/WaveformTransformations/Private/WaveformTransformationNormalize.cpp b/Engine/Plugins/Editor/WaveformEditor/Source/WaveformTransformations/Private/WaveformTransformationNormalize.cpp
--- a/Engine/Plugins/Editor/WaveformEditor/Source/WaveformTransformations/Private/WaveformTransformationNormalize.cpp
+++ b/Engine/Plugins/Editor/WaveformEditor/Source/WaveformTransformations/Private/WaveformTransformationNormalize.cpp
@@ -103,6 +103,7 @@ FWaveTransformationNormalize::FWaveTransformationNormalize(float InTarget, float
 		}
 	case ENormalizationMode::RMS:
 		PeakDecibelValue = GetRMSPeak(InputAudio, InOutWaveInfo.SampleRate, InOutWaveInfo.NumChannels);
+		PeakDecibelValue = Audio::ConvertToDecibels(PeakDecibelValue);
 		break;
 	case ENormalizationMode::DWeightedLoudness:
 		PeakDecibelValue = GetLoudnessPeak(InputAudio,  InOutWaveInfo.SampleRate, InOutWaveInfo.NumChannels);
@@ -124,4 +125,4 @@ FWaveTransformationNormalize::FWaveTransformationNormalize(float InTarget, float
 Audio::FTransformationPtr UWaveformTransformationNormalize::CreateTransformation() const
 {
 	return MakeUnique<FWaveTransformationNormalize>(Target, MaxGain, Mode);
-}
\ No newline at end of file
+}
