[WaveformEditor] Use FGeometry::GetLocalSize() in widgets to support windows scaling
#trivial

diff --git a/Engine/Plugins/Editor/WaveformEditor/Source/WaveformEditorWidgets/Private/SWaveformPanel.cpp b/Engine/Plugins/Editor/WaveformEditor/Source/WaveformEditorWidgets/Private/SWaveformPanel.cpp
--- a/Engine/Plugins/Editor/WaveformEditor/Source/WaveformEditorWidgets/Private/SWaveformPanel.cpp
+++ b/Engine/Plugins/Editor/WaveformEditor/Source/WaveformEditorWidgets/Private/SWaveformPanel.cpp
@@ -128,7 +128,7 @@ void SWaveformPanel::UpdateDisplayUnit(const EWaveformEditorDisplayUnit InDispla
 
 void SWaveformPanel::Tick(const FGeometry& AllottedGeometry, const double InCurrentTime, const float InDeltaTime)
 {
-	const float PaintedWidth = AllottedGeometry.GetAbsoluteSize().X;
+	const float PaintedWidth = AllottedGeometry.GetLocalSize().X;
 
 	if (PaintedWidth != CachedPixelWidth)
 	{
diff --git a/Engine/Plugins/Editor/WaveformEditor/Source/WaveformEditorWidgets/Private/SWaveformViewer.cpp b/Engine/Plugins/Editor/WaveformEditor/Source/WaveformEditorWidgets/Private/SWaveformViewer.cpp
--- a/Engine/Plugins/Editor/WaveformEditor/Source/WaveformEditorWidgets/Private/SWaveformViewer.cpp
+++ b/Engine/Plugins/Editor/WaveformEditor/Source/WaveformEditorWidgets/Private/SWaveformViewer.cpp
@@ -140,7 +140,7 @@ FVector2D SWaveformViewer::ComputeDesiredSize(float) const
 
 void SWaveformViewer::Tick(const FGeometry& AllottedGeometry, const double InCurrentTime, const float InDeltaTime)
 {
-	uint32 DiscretePixelWidth = FMath::FloorToInt(AllottedGeometry.GetAbsoluteSize().X);
+	uint32 DiscretePixelWidth = FMath::FloorToInt(AllottedGeometry.GetLocalSize().X);
 	if (DiscretePixelWidth <= 0)
 	{
 		return;
@@ -149,7 +149,7 @@ void SWaveformViewer::Tick(const FGeometry& AllottedGeometry, const double InCur
 	if (DiscretePixelWidth != CachedPixelWidth || bForceRedraw)
 	{
 		CachedPixelWidth = DiscretePixelWidth;
-		CachedPixelHeight = AllottedGeometry.GetAbsoluteSize().Y;
+		CachedPixelHeight = AllottedGeometry.GetLocalSize().Y;
 		const float NumSamplesToDisplay = (DisplayRange.GetUpperBoundValue() - DisplayRange.GetLowerBoundValue()) * RenderData->GetSampleRate();
 		WaveformDrawMode = NumSamplesToDisplay > CachedPixelWidth ? EWaveformDrawMode::BinnedPeaks : EWaveformDrawMode::WaveformLine;
 
@@ -165,9 +165,9 @@ void SWaveformViewer::Tick(const FGeometry& AllottedGeometry, const double InCur
 		
 		bForceRedraw = false;
 	}
-	else if (CachedPixelHeight != AllottedGeometry.GetAbsoluteSize().Y)
+	else if (CachedPixelHeight != AllottedGeometry.GetLocalSize().Y)
 	{
-		CachedPixelHeight = AllottedGeometry.GetAbsoluteSize().Y;
+		CachedPixelHeight = AllottedGeometry.GetLocalSize().Y;
 
 		if (WaveformDrawMode == EWaveformDrawMode::BinnedPeaks)
 		{
@@ -190,10 +190,10 @@ void SWaveformViewer::GenerateWaveformBinsCoordinates(TArray<FSamplesBinCoordina
 		return;
 	}
 
-	const uint32 PixelWidth = FMath::FloorToInt(InAllottedGeometry.GetAbsoluteSize().X);
+	const uint32 PixelWidth = FMath::FloorToInt(InAllottedGeometry.GetLocalSize().X);
 	check(PixelWidth * NChannels == InWaveformPeaks.Num());
 
-	const float HeightScale = InAllottedGeometry.GetAbsoluteSize().Y / (2.f * TNumericLimits<int16>::Max() * NChannels) * MaxWaveformHeight * VerticalZoomFactor;
+	const float HeightScale = InAllottedGeometry.GetLocalSize().Y / (2.f * TNumericLimits<int16>::Max() * NChannels) * MaxWaveformHeight * VerticalZoomFactor;
 
 	OutDrawCoordinates.SetNumUninitialized(InWaveformPeaks.Num());
 	FSamplesBinCoordinates* OutCoordinatesData = OutDrawCoordinates.GetData();
diff --git a/Engine/Plugins/Editor/WaveformEditor/Source/WaveformEditorWidgets/Private/SWaveformViewer.h b/Engine/Plugins/Editor/WaveformEditor/Source/WaveformEditorWidgets/Private/SWaveformViewer.h
--- a/Engine/Plugins/Editor/WaveformEditor/Source/WaveformEditorWidgets/Private/SWaveformViewer.h
+++ b/Engine/Plugins/Editor/WaveformEditor/Source/WaveformEditorWidgets/Private/SWaveformViewer.h
@@ -24,7 +24,7 @@ struct FChannelSlotBoundaries
 {
 	explicit FChannelSlotBoundaries(const uint16 ChannelToDraw, const uint16 TotalNumChannels, const FGeometry& InAllottedGeometry)
 	{
-		Height = InAllottedGeometry.GetAbsoluteSize().Y / TotalNumChannels;
+		Height = InAllottedGeometry.GetLocalSize().Y / TotalNumChannels;
 		float ChannelSlotMidPoint = Height / 2.f;
 
 		Top = Height * ChannelToDraw;
