[WaveformTransformationNormalize] move helper methods from static to namespace
#trivial

diff --git a/Engine/Plugins/Editor/WaveformEditor/Source/WaveformTransformations/Private/WaveformTransformationNormalize.cpp b/Engine/Plugins/Editor/WaveformEditor/Source/WaveformTransformations/Private/WaveformTransformationNormalize.cpp
--- a/Engine/Plugins/Editor/WaveformEditor/Source/WaveformTransformations/Private/WaveformTransformationNormalize.cpp
+++ b/Engine/Plugins/Editor/WaveformEditor/Source/WaveformTransformations/Private/WaveformTransformationNormalize.cpp
@@ -9,76 +9,80 @@
 #include "Algo/MaxElement.h"
 #include "MeterFactory.h"
 
-static float GetRMSPeak(const Audio::FAlignedFloatBuffer& InputAudio, float SampleRate, int32 NumChannels)
+namespace WaveformTransformNormalizeHelpers
 {
-	Audio::FMeterFactory Analyzer;
-	Audio::FMeterSettings Settings;
-	Settings.AnalysisPeriod = 0.3f;
-
-	TUniquePtr<Audio::IAnalyzerResult> Result = Analyzer.NewResult();
-	TUniquePtr<Audio::IAnalyzerWorker> AnalyzerWorker = Analyzer.NewWorker({(int32)SampleRate, NumChannels}, &Settings);
-
-	if(AnalyzerWorker == nullptr)
+	float GetRMSPeak(const Audio::FAlignedFloatBuffer& InputAudio, float SampleRate, int32 NumChannels)
 	{
-		return 0.f;
-	}
-	
-	AnalyzerWorker->Analyze(MakeArrayView(InputAudio), Result.Get());
-
-	Audio::FMeterResult* MeterResult = static_cast<Audio::FMeterResult*>(Result.Get());
+		Audio::FMeterFactory Analyzer;
+		Audio::FMeterSettings Settings;
+		Settings.AnalysisPeriod = 0.3f;
 
-	if (MeterResult)
-	{
-		const TArray<Audio::FMeterEntry>& LoudnessArray = MeterResult->GetMeterArray();
-		auto MeterEntry = Algo::MaxElement(LoudnessArray, [](const Audio::FMeterEntry& A, const Audio::FMeterEntry& B)
-		{
-			return A.MeterValue < B.MeterValue;
-		});
+		TUniquePtr<Audio::IAnalyzerResult> Result = Analyzer.NewResult();
+		TUniquePtr<Audio::IAnalyzerWorker> AnalyzerWorker = Analyzer.NewWorker({ (int32)SampleRate, NumChannels }, &Settings);
 
-		if (MeterEntry)
+		if (AnalyzerWorker == nullptr)
 		{
-			return MeterEntry->MeterValue;
+			return 0.f;
 		}
-	}
 
-	return 0.f;
-}
+		AnalyzerWorker->Analyze(MakeArrayView(InputAudio), Result.Get());
 
-static float GetLoudnessPeak(const Audio::FAlignedFloatBuffer& InputAudio, float SampleRate, int32 NumChannels)
-{
-	Audio::FLoudnessFactory Analyzer;
-	Audio::FLoudnessSettings Settings;
-	Settings.AnalysisPeriod = 1.f;
-	
-	TUniquePtr<Audio::IAnalyzerResult> Result = Analyzer.NewResult();
-	TUniquePtr<Audio::IAnalyzerWorker> AnalyzerWorker = Analyzer.NewWorker({(int32)SampleRate, NumChannels}, &Settings);
+		Audio::FMeterResult* MeterResult = static_cast<Audio::FMeterResult*>(Result.Get());
+
+		if (MeterResult)
+		{
+			const TArray<Audio::FMeterEntry>& LoudnessArray = MeterResult->GetMeterArray();
+			auto MeterEntry = Algo::MaxElement(LoudnessArray, [](const Audio::FMeterEntry& A, const Audio::FMeterEntry& B)
+				{
+					return A.MeterValue < B.MeterValue;
+				});
+
+			if (MeterEntry)
+			{
+				return MeterEntry->MeterValue;
+			}
+		}
 
-	if(AnalyzerWorker == nullptr)
-	{
 		return 0.f;
 	}
-	
-	AnalyzerWorker->Analyze(MakeArrayView(InputAudio), Result.Get());
-
-	Audio::FLoudnessResult* LoudnessResult = static_cast<Audio::FLoudnessResult*>(Result.Get());
 
-	if (LoudnessResult)
+	float GetLoudnessPeak(const Audio::FAlignedFloatBuffer& InputAudio, float SampleRate, int32 NumChannels)
 	{
-		const TArray<Audio::FLoudnessEntry>& LoudnessArray = LoudnessResult->GetLoudnessArray();
-		auto LoudestEntry = Algo::MaxElement(LoudnessArray, [](const Audio::FLoudnessEntry& A, const Audio::FLoudnessEntry& B)
+		Audio::FLoudnessFactory Analyzer;
+		Audio::FLoudnessSettings Settings;
+		Settings.AnalysisPeriod = 1.f;
+
+		TUniquePtr<Audio::IAnalyzerResult> Result = Analyzer.NewResult();
+		TUniquePtr<Audio::IAnalyzerWorker> AnalyzerWorker = Analyzer.NewWorker({ (int32)SampleRate, NumChannels }, &Settings);
+
+		if (AnalyzerWorker == nullptr)
 		{
-			return A.Loudness < B.Loudness;
-		});
+			return 0.f;
+		}
 
-		if (LoudestEntry)
+		AnalyzerWorker->Analyze(MakeArrayView(InputAudio), Result.Get());
+
+		Audio::FLoudnessResult* LoudnessResult = static_cast<Audio::FLoudnessResult*>(Result.Get());
+
+		if (LoudnessResult)
 		{
-			return LoudestEntry->Loudness;
+			const TArray<Audio::FLoudnessEntry>& LoudnessArray = LoudnessResult->GetLoudnessArray();
+			auto LoudestEntry = Algo::MaxElement(LoudnessArray, [](const Audio::FLoudnessEntry& A, const Audio::FLoudnessEntry& B)
+				{
+					return A.Loudness < B.Loudness;
+				});
+
+			if (LoudestEntry)
+			{
+				return LoudestEntry->Loudness;
+			}
 		}
-	}
 
-	return 0.f;
+		return 0.f;
+	}
 }
 
+
 FWaveTransformationNormalize::FWaveTransformationNormalize(float InTarget, float InMaxGain, ENormalizationMode InMode)
 	: Target(InTarget)
 	, MaxGain(InMaxGain)
@@ -102,11 +106,11 @@ FWaveTransformationNormalize::FWaveTransformationNormalize(float InTarget, float
 			break;
 		}
 	case ENormalizationMode::RMS:
-		PeakDecibelValue = GetRMSPeak(InputAudio, InOutWaveInfo.SampleRate, InOutWaveInfo.NumChannels);
+		PeakDecibelValue = WaveformTransformNormalizeHelpers::GetRMSPeak(InputAudio, InOutWaveInfo.SampleRate, InOutWaveInfo.NumChannels);
 		PeakDecibelValue = Audio::ConvertToDecibels(PeakDecibelValue);
 		break;
 	case ENormalizationMode::DWeightedLoudness:
-		PeakDecibelValue = GetLoudnessPeak(InputAudio,  InOutWaveInfo.SampleRate, InOutWaveInfo.NumChannels);
+		PeakDecibelValue = WaveformTransformNormalizeHelpers::GetLoudnessPeak(InputAudio,  InOutWaveInfo.SampleRate, InOutWaveInfo.NumChannels);
 		break;
 	case ENormalizationMode::COUNT:
 	default:
