[WaveformTransformationNormalize] Use entire audio as buffer for RMS analysis when  input audio length would be < RMSAnalysisPeriod
fix for #jira UE-165413

diff --git a/Engine/Plugins/Editor/WaveformEditor/Source/WaveformTransformations/Private/WaveformTransformationNormalize.cpp b/Engine/Plugins/Editor/WaveformEditor/Source/WaveformTransformations/Private/WaveformTransformationNormalize.cpp
--- a/Engine/Plugins/Editor/WaveformEditor/Source/WaveformTransformations/Private/WaveformTransformationNormalize.cpp
+++ b/Engine/Plugins/Editor/WaveformEditor/Source/WaveformTransformations/Private/WaveformTransformationNormalize.cpp
@@ -13,36 +13,57 @@ namespace WaveformTransformNormalizeHelpers
 {
 	float GetRMSPeak(const Audio::FAlignedFloatBuffer& InputAudio, float SampleRate, int32 NumChannels)
 	{
-		Audio::FMeterFactory Analyzer;
-		Audio::FMeterSettings Settings;
-		Settings.AnalysisPeriod = 0.3f;
+		const float AudioLengthSeconds = InputAudio.Num() / NumChannels / SampleRate;
+		const float RMSAnalysisPeriod = 0.3f;
+		const bool bAnalyzeAsSingleBuffer = RMSAnalysisPeriod > AudioLengthSeconds;
 
-		TUniquePtr<Audio::IAnalyzerResult> Result = Analyzer.NewResult();
-		TUniquePtr<Audio::IAnalyzerWorker> AnalyzerWorker = Analyzer.NewWorker({ (int32)SampleRate, NumChannels }, &Settings);
+		Audio::FMeterSettings Settings;
+		Settings.AnalysisPeriod = bAnalyzeAsSingleBuffer ? AudioLengthSeconds : RMSAnalysisPeriod;
 
-		if (AnalyzerWorker == nullptr)
+		if (bAnalyzeAsSingleBuffer)
 		{
-			return 0.f;
+			Audio::FMeterAnalyzer Analyzer(SampleRate, NumChannels, Settings);
+			Audio::FMeterAnalyzerResults Results = Analyzer.ProcessAudio(InputAudio);
+			const float* MeterValue = Algo::MaxElement(Results.MeterValues, [](const float A, const float B)
+				{
+					return A < B;
+				});
+
+			if (MeterValue)
+			{
+				return *MeterValue;
+			}
 		}
+		else
+		{
+			Audio::FMeterFactory Analyzer;
+			TUniquePtr<Audio::IAnalyzerResult> Result = Analyzer.NewResult();
+			TUniquePtr<Audio::IAnalyzerWorker> AnalyzerWorker = Analyzer.NewWorker({ (int32)SampleRate, NumChannels }, &Settings);
 
-		AnalyzerWorker->Analyze(MakeArrayView(InputAudio), Result.Get());
+			if (AnalyzerWorker == nullptr)
+			{
+				return 0.f;
+			}
 
-		Audio::FMeterResult* MeterResult = static_cast<Audio::FMeterResult*>(Result.Get());
+			AnalyzerWorker->Analyze(MakeArrayView(InputAudio), Result.Get());
 
-		if (MeterResult)
-		{
-			const TArray<Audio::FMeterEntry>& LoudnessArray = MeterResult->GetMeterArray();
-			auto MeterEntry = Algo::MaxElement(LoudnessArray, [](const Audio::FMeterEntry& A, const Audio::FMeterEntry& B)
-				{
-					return A.MeterValue < B.MeterValue;
-				});
+			Audio::FMeterResult* MeterResult = static_cast<Audio::FMeterResult*>(Result.Get());
 
-			if (MeterEntry)
+			if (MeterResult)
 			{
-				return MeterEntry->MeterValue;
+				const TArray<Audio::FMeterEntry>& LoudnessArray = MeterResult->GetMeterArray();
+				auto MeterEntry = Algo::MaxElement(LoudnessArray, [](const Audio::FMeterEntry& A, const Audio::FMeterEntry& B)
+					{
+						return A.MeterValue < B.MeterValue;
+					});
+
+				if (MeterEntry)
+				{
+					return MeterEntry->MeterValue;
+				}
 			}
 		}
-
+		
 		return 0.f;
 	}
 
diff --git a/Engine/Plugins/Runtime/AudioSynesthesia/Source/AudioSynesthesiaCore/Public/MeterAnalyzer.h b/Engine/Plugins/Runtime/AudioSynesthesia/Source/AudioSynesthesiaCore/Public/MeterAnalyzer.h
--- a/Engine/Plugins/Runtime/AudioSynesthesia/Source/AudioSynesthesiaCore/Public/MeterAnalyzer.h
+++ b/Engine/Plugins/Runtime/AudioSynesthesia/Source/AudioSynesthesiaCore/Public/MeterAnalyzer.h
@@ -25,7 +25,7 @@ namespace Audio
 	};
 
 	// Per-channel analyzer results
-	struct FMeterAnalyzerResults
+	struct AUDIOSYNESTHESIACORE_API FMeterAnalyzerResults
 	{
 		float TimeSec = 0.0f;
 		TArray<float> MeterValues;
@@ -34,7 +34,7 @@ namespace Audio
 		TArray<int32> NumSamplesClipping;
 	};
 
-	class FMeterAnalyzer
+	class AUDIOSYNESTHESIACORE_API FMeterAnalyzer
 	{
 	public:
 
