[WaveformEditor] Handle playhead/transport behavior when transformations are added
fix for #jira UE-161066

diff --git a/Engine/Plugins/Editor/WaveformEditor/Source/WaveformEditor/Private/WaveformEditor.cpp b/Engine/Plugins/Editor/WaveformEditor/Source/WaveformEditor/Private/WaveformEditor.cpp
--- a/Engine/Plugins/Editor/WaveformEditor/Source/WaveformEditor/Private/WaveformEditor.cpp
+++ b/Engine/Plugins/Editor/WaveformEditor/Source/WaveformEditor/Private/WaveformEditor.cpp
@@ -35,6 +35,7 @@ bool FWaveformEditor::Init(const EToolkitMode::Type Mode, const TSharedPtr<ITool
 	checkf(SoundWaveToEdit, TEXT("Tried to open a Soundwave Editor from a null soundwave"));
 
 	SoundWave = SoundWaveToEdit;
+
 	AddDefaultTransformations();
 
 	bool bIsInitialized = true;
@@ -105,20 +106,7 @@ void FWaveformEditor::AddDefaultTransformations()
 			}
 		}
 
-		FWaveTransformUObjectConfiguration ChainMetadata = FWaveTransformUObjectConfiguration();
-
-		ChainMetadata.NumChannels = SoundWave->NumChannels;
-		ChainMetadata.SampleRate = SoundWave->GetSampleRateForCurrentPlatform();
-		ChainMetadata.StartTime = 0.f;
-		ChainMetadata.EndTime = SoundWave->Duration;
-
-		for (TObjectPtr<UWaveformTransformationBase>& Transform : SoundWave->Transformations)
-		{
-			if (Transform)
-			{
-				Transform->UpdateConfiguration(ChainMetadata);
-			}
-		}
+		TransformationChainConfig = SoundWave->UpdateTransformations();
 	}
 }
 
@@ -438,6 +426,44 @@ void FWaveformEditor::NotifyPostChange(const FPropertyChangedEvent& PropertyChan
 	{
 		return;
 	}
+	
+
+	if (PropertyChangedEvent.ChangeType == EPropertyChangeType::Interactive)
+	{
+		bIsInteractingWithTransformations = true;
+
+		if (TransportController->IsPlaying() || TransportController->IsPaused())
+		{
+			TransformInteractionPlayState = AudioComponent->GetPlayState();
+			PlaybackTimeBeforeTransformInteraction = LastReceivedPlaybackPercent * (TransformationChainConfig.EndTime - TransformationChainConfig.StartTime);
+			StartTimeBeforeTransformInteraction = TransformationChainConfig.StartTime;
+			AudioComponent->Stop();
+		}
+	}
+
+	if (PropertyChangedEvent.ChangeType == EPropertyChangeType::ValueSet)
+	{
+		if (!TransportController->IsPlaying())
+		{
+			const float StartTimeDifference = StartTimeBeforeTransformInteraction - TransformationChainConfig.StartTime;
+			const float AdjustedPlaybackTime = PlaybackTimeBeforeTransformInteraction + StartTimeDifference;
+
+			switch (TransformInteractionPlayState)
+			{
+			case EAudioComponentPlayState::Playing:
+				TransportController->Play(AdjustedPlaybackTime);
+				TransformInteractionPlayState = EAudioComponentPlayState::Stopped;
+				break;
+			case EAudioComponentPlayState::Paused:
+				TransportController->CacheStartTime(AdjustedPlaybackTime);
+				TransformInteractionPlayState = EAudioComponentPlayState::Stopped;
+			default:
+				break;
+			}
+		}
+
+		bIsInteractingWithTransformations = false;
+	}
 
 	const bool bUpdateTransformationChain = PropertyChangedEvent.GetPropertyName() == TEXT("Transformations");
 
@@ -447,8 +473,11 @@ void FWaveformEditor::NotifyPostChange(const FPropertyChangedEvent& PropertyChan
 	}
 
 	TransformationsRenderManager->UpdateRenderElements();
+	TransformationChainConfig = SoundWave->GetTransformationChainConfig();
 }
 
+
+
 void FWaveformEditor::PostUndo(bool bSuccess)
 {
 	if (bSuccess)
@@ -622,6 +651,7 @@ void FWaveformEditor::HandlePlaybackPercentageChange(const UAudioComponent* InCo
 	const bool bIsStopped = AudioComponent->GetPlayState() == EAudioComponentPlayState::Stopped;
 	const bool bIsPaused = AudioComponent->GetPlayState() == EAudioComponentPlayState::Paused;
 	const bool bPropagatePercentage = !bIsStopped && !bIsPaused;
+	LastReceivedPlaybackPercent = InPlaybackPercentage;
 	
 	if (InComponent == AudioComponent && bPropagatePercentage)
 	{
@@ -645,11 +675,14 @@ void FWaveformEditor::HandleAudioComponentPlayStateChanged(const UAudioComponent
 	default:
 		break;
 	case EAudioComponentPlayState::Stopped:
-		if (!TransportCoordinator->IsScrubbing())
+		const bool bIgnoreStopState = TransportCoordinator->IsScrubbing() || bIsInteractingWithTransformations;
+
+		if (!bIgnoreStopState)
 		{
 			TransportController->CacheStartTime(0.f);
 			TransportCoordinator->Stop();
 		}
+
 		break;
 	}
 }
diff --git a/Engine/Plugins/Editor/WaveformEditor/Source/WaveformEditor/Public/WaveformEditor.h b/Engine/Plugins/Editor/WaveformEditor/Source/WaveformEditor/Public/WaveformEditor.h
--- a/Engine/Plugins/Editor/WaveformEditor/Source/WaveformEditor/Public/WaveformEditor.h
+++ b/Engine/Plugins/Editor/WaveformEditor/Source/WaveformEditor/Public/WaveformEditor.h
@@ -37,6 +37,7 @@ public:
 	/** FNotifyHook interface */
 	void NotifyPreChange(class FEditPropertyChain* PropertyAboutToChange) override {};
 	void NotifyPostChange(const FPropertyChangedEvent& PropertyChangedEvent, class FEditPropertyChain* PropertyThatChanged) override;
+	
 
 	/** FEditorUndo interface */
 	void PostUndo(bool bSuccess) override;
@@ -127,4 +128,13 @@ private:
 	UAudioComponent* AudioComponent = nullptr;
 
 	bool bWasPlayingBeforeScrubbing = false;
+	bool bIsInteractingWithTransformations = false;
+
+	float LastReceivedPlaybackPercent = 0.f;
+
+	EAudioComponentPlayState TransformInteractionPlayState = EAudioComponentPlayState::Stopped;
+	float PlaybackTimeBeforeTransformInteraction = 0.f;
+	float StartTimeBeforeTransformInteraction = 0.f;
+
+	FWaveTransformUObjectConfiguration TransformationChainConfig;
 };
diff --git a/Engine/Plugins/Editor/WaveformEditor/Source/WaveformEditorWidgets/Private/WaveformEditorTransportCoordinator.cpp b/Engine/Plugins/Editor/WaveformEditor/Source/WaveformEditorWidgets/Private/WaveformEditorTransportCoordinator.cpp
--- a/Engine/Plugins/Editor/WaveformEditor/Source/WaveformEditorWidgets/Private/WaveformEditorTransportCoordinator.cpp
+++ b/Engine/Plugins/Editor/WaveformEditor/Source/WaveformEditorWidgets/Private/WaveformEditorTransportCoordinator.cpp
@@ -118,11 +118,9 @@ FReply FWaveformEditorTransportCoordinator::HandleTimeRulerInteraction(const ERe
 
 void FWaveformEditorTransportCoordinator::HandleRenderDataUpdate()
 {
-	PlaybackRange = RenderData->GetTransformedWaveformBounds();
-	ReceivePlayBackRatio(0.f);
+ 	PlaybackRange = RenderData->GetTransformedWaveformBounds();
 }
 
-
 void FWaveformEditorTransportCoordinator::ScrubPlayhead(const float TargetPlayheadPosition, const bool bIsMoving)
 {
 	const float LowerPlaybackBoundPosition = (PlaybackRange.GetLowerBoundValue() - DisplayRange.GetLowerBoundValue()) / ZoomRatio;
@@ -285,7 +283,6 @@ void FWaveformEditorTransportCoordinator::UpdateDisplayRange(const float MinValu
 	{
 		OnDisplayRangeUpdated.Broadcast(DisplayRange);
 	}
-
 }
 
 bool FWaveformEditorTransportCoordinator::IsRatioWithinDisplayRange(const float Ratio) const
diff --git a/Engine/Source/Runtime/Engine/Classes/Sound/SoundWave.h b/Engine/Source/Runtime/Engine/Classes/Sound/SoundWave.h
--- a/Engine/Source/Runtime/Engine/Classes/Sound/SoundWave.h
+++ b/Engine/Source/Runtime/Engine/Classes/Sound/SoundWave.h
@@ -462,7 +462,7 @@ public:
 
 	/** Returns the Runtime format of the wave */
 	FName GetRuntimeFormat() const;
-		
+
 private:
 	// cached proxy
 	FSoundWaveProxyPtr InternalProxy{ nullptr };
@@ -571,6 +571,11 @@ public:
 	UPROPERTY(EditAnywhere, AdvancedDisplay, Category = "Loading")
 	int32 InitialChunkSize;
 
+#if WITH_EDITOR
+	const FWaveTransformUObjectConfiguration& GetTransformationChainConfig() const;
+	const FWaveTransformUObjectConfiguration& UpdateTransformations();
+#endif
+
 private:
 
 	/** Helper functions to search analysis data. Takes starting index to start query. Returns which data index the result was found at. Returns INDEX_NONE if not found. */
@@ -595,7 +600,6 @@ private:
 	// Called when we change any properties about the underlying audio asset
 #if WITH_EDITOR
 	void UpdateAsset();
-	void UpdateTransformations();
 #endif
 
 public:
@@ -1035,6 +1039,8 @@ private:
 
 	void BakeFFTAnalysis();
 	void BakeEnvelopeAnalysis();
+
+	FWaveTransformUObjectConfiguration TransformationChainConfig;
 #endif //WITH_EDITOR
 
 public:
diff --git a/Engine/Source/Runtime/Engine/Private/SoundWave.cpp b/Engine/Source/Runtime/Engine/Private/SoundWave.cpp
--- a/Engine/Source/Runtime/Engine/Private/SoundWave.cpp
+++ b/Engine/Source/Runtime/Engine/Private/SoundWave.cpp
@@ -3894,29 +3894,36 @@ void USoundWave::UpdateAsset()
 	}
 }
 
-void USoundWave::UpdateTransformations()
+const FWaveTransformUObjectConfiguration& USoundWave::GetTransformationChainConfig() const
 {
-	FWaveTransformUObjectConfiguration ChainConfiguration = FWaveTransformUObjectConfiguration();
+	return TransformationChainConfig;
+}
+
+const FWaveTransformUObjectConfiguration& USoundWave::UpdateTransformations()
+{
+	TransformationChainConfig = FWaveTransformUObjectConfiguration();
 
 	if (TotalSamples == 0)
 	{
 		UE_LOG(LogAudio, Warning, TEXT("Found %d total samples when updating transformations on soundwave %s"), TotalSamples, *GetNameSafe(this));
 		UE_LOG(LogAudio, Warning, TEXT("Consider reimporting the asset, transformations might not work correctly"));
-		return;
+		return TransformationChainConfig;
 	}
 
-	ChainConfiguration.NumChannels = NumChannels;
-	ChainConfiguration.SampleRate = ImportedSampleRate > 0 ? ImportedSampleRate : GetSampleRateForCurrentPlatform();
-	ChainConfiguration.StartTime = 0.f;
-	ChainConfiguration.EndTime = TotalSamples / ChainConfiguration.SampleRate;
+	TransformationChainConfig.NumChannels = NumChannels;
+	TransformationChainConfig.SampleRate = ImportedSampleRate > 0 ? ImportedSampleRate : GetSampleRateForCurrentPlatform();
+	TransformationChainConfig.StartTime = 0.f;
+	TransformationChainConfig.EndTime = TotalSamples / TransformationChainConfig.SampleRate;
 
 	for (TObjectPtr<UWaveformTransformationBase>& Transformation : Transformations)
 	{
 		if (Transformation)
 		{
-			Transformation->UpdateConfiguration(ChainConfiguration);
+			Transformation->UpdateConfiguration(TransformationChainConfig);
 		}
 	}
+
+	return TransformationChainConfig;
 }
 
 #endif // #if WITH_EDITOR
