[Waveform Transformations] Add Tooltip and category to waveform transformations property on USoundWave and make it show up only if the relevant plugin module is loaded

diff --git a/Engine/Plugins/Editor/WaveformEditor/Source/WaveformEditor/Private/WaveformEditor.cpp b/Engine/Plugins/Editor/WaveformEditor/Source/WaveformEditor/Private/WaveformEditor.cpp
--- a/Engine/Plugins/Editor/WaveformEditor/Source/WaveformEditor/Private/WaveformEditor.cpp
+++ b/Engine/Plugins/Editor/WaveformEditor/Source/WaveformEditor/Private/WaveformEditor.cpp
@@ -188,7 +188,7 @@ void FWaveformEditor::RegisterTabSpawners(const TSharedRef<FTabManager>& InTabMa
 		.SetIcon(FSlateIcon(FAppStyle::GetAppStyleSetName(), "LevelEditor.Tabs.Details"));
 
 	InTabManager->RegisterTabSpawner(TransformationsTabId, FOnSpawnTab::CreateSP(this, &FWaveformEditor::SpawnTab_Transformations))
-		.SetDisplayName(LOCTEXT("TransformationsTab", "Transformations"))
+		.SetDisplayName(LOCTEXT("ProcessingTab", "Processing"))
 		.SetGroup(WorkspaceMenuCategory.ToSharedRef())
 		.SetIcon(FSlateIcon(FAppStyle::GetAppStyleSetName(), "LevelEditor.Tabs.Details"));
 
@@ -604,7 +604,7 @@ TSharedRef<SDockTab> FWaveformEditor::SpawnTab_Transformations(const FSpawnTabAr
 	check(Args.GetTabId() == TransformationsTabId);
 
 	return SNew(SDockTab)
-		.Label(LOCTEXT("SoundWaveTransformationsTitle", "Transformations"))
+		.Label(LOCTEXT("SoundWaveProcessingTitle", "Processing"))
 		[
 			TransformationsDetails.ToSharedRef()
 		];
diff --git a/Engine/Plugins/Editor/WaveformEditor/Source/WaveformEditor/Private/WaveformEditorDetailsCustomization.cpp b/Engine/Plugins/Editor/WaveformEditor/Source/WaveformEditor/Private/WaveformEditorDetailsCustomization.cpp
--- a/Engine/Plugins/Editor/WaveformEditor/Source/WaveformEditor/Private/WaveformEditorDetailsCustomization.cpp
+++ b/Engine/Plugins/Editor/WaveformEditor/Source/WaveformEditor/Private/WaveformEditorDetailsCustomization.cpp
@@ -6,7 +6,7 @@
 #include "DetailLayoutBuilder.h"
 #include "Sound/SoundWave.h"
 
-static const FLazyName TransformationsCategoryName("Transformations");
+static const FLazyName TransformationsCategoryName("Waveform Processing");
 
 void FWaveformTransformationsDetailsCustomization::CustomizeDetails(IDetailLayoutBuilder& DetailLayout)
 {	
@@ -21,7 +21,7 @@ void FWaveformTransformationsDetailsCustomization::CustomizeDetails(IDetailLayou
 		}
 	}
 
-	IDetailCategoryBuilder& CategoryBuilder = DetailLayout.EditCategory("Transformations");
+	IDetailCategoryBuilder& CategoryBuilder = DetailLayout.EditCategory("Waveform Processing");
 	CategoryBuilder.InitiallyCollapsed(false);
 	CategoryBuilder.RestoreExpansionState(true);
 }
\ No newline at end of file
diff --git a/Engine/Plugins/Runtime/Metasound/Source/MetasoundEditor/Private/MetasoundDetailCustomization.cpp b/Engine/Plugins/Runtime/Metasound/Source/MetasoundEditor/Private/MetasoundDetailCustomization.cpp
--- a/Engine/Plugins/Runtime/Metasound/Source/MetasoundEditor/Private/MetasoundDetailCustomization.cpp
+++ b/Engine/Plugins/Runtime/Metasound/Source/MetasoundEditor/Private/MetasoundDetailCustomization.cpp
@@ -317,7 +317,7 @@ namespace Metasound
 			DetailLayout.HideCategory("Loading");
 			DetailLayout.HideCategory("Playback");
 			DetailLayout.HideCategory("Subtitles");
-			DetailLayout.HideCategory("Transformations");
+			DetailLayout.HideCategory("Waveform Processing");;
 		}
 
 		FMetasoundInterfacesDetailCustomization::FMetasoundInterfacesDetailCustomization()
diff --git a/Engine/Source/Editor/DetailCustomizations/Private/SoundSourceBusDetails.cpp b/Engine/Source/Editor/DetailCustomizations/Private/SoundSourceBusDetails.cpp
--- a/Engine/Source/Editor/DetailCustomizations/Private/SoundSourceBusDetails.cpp
+++ b/Engine/Source/Editor/DetailCustomizations/Private/SoundSourceBusDetails.cpp
@@ -14,8 +14,7 @@ TSharedRef<IDetailCustomization> FSoundSourceBusDetails::MakeInstance()
 void FSoundSourceBusDetails::CustomizeDetails(IDetailLayoutBuilder& DetailBuilder)
 {
 	// Hide details of sound base that aren't relevant for buses
-	DetailBuilder.HideCategory("Transformations");
-
+	DetailBuilder.HideCategory("Waveform Processing");
 }
 
 #undef LOCTEXT_NAMESPACE
diff --git a/Engine/Source/Editor/DetailCustomizations/Private/SoundWaveDetails.cpp b/Engine/Source/Editor/DetailCustomizations/Private/SoundWaveDetails.cpp
--- a/Engine/Source/Editor/DetailCustomizations/Private/SoundWaveDetails.cpp
+++ b/Engine/Source/Editor/DetailCustomizations/Private/SoundWaveDetails.cpp
@@ -45,6 +45,11 @@ TSharedRef<IDetailCustomization> FSoundWaveDetails::MakeInstance()
 void FSoundWaveDetails::CustomizeDetails(IDetailLayoutBuilder& DetailBuilder)
 {
 	CustomizeCurveDetails(DetailBuilder);
+
+	if(!FModuleManager::Get().IsModuleLoaded("WaveformTransformations"))
+	{
+		DetailBuilder.HideCategory("Waveform Processing");
+	}
 }
 
 void FSoundWaveDetails::CustomizeCurveDetails(IDetailLayoutBuilder& DetailBuilder)
diff --git a/Engine/Source/Runtime/Engine/Classes/Sound/SoundWave.h b/Engine/Source/Runtime/Engine/Classes/Sound/SoundWave.h
--- a/Engine/Source/Runtime/Engine/Classes/Sound/SoundWave.h
+++ b/Engine/Source/Runtime/Engine/Classes/Sound/SoundWave.h
@@ -817,7 +817,8 @@ public:
 	/** Uncompressed wav data 16 bit in mono or stereo - stereo not allowed for multichannel data */
 	UE::Serialization::FEditorBulkData RawData;
 
-	UPROPERTY(EditAnywhere, Instanced, Category = "Transformations")
+	/** Waveform edits to be applied to this SoundWave on cook (editing transformations will trigger a cook) */
+	UPROPERTY(EditAnywhere, Instanced, Category = "Waveform Processing")
 	TArray<TObjectPtr<class UWaveformTransformationBase>> Transformations;
 #endif
 
