Add version number to waveform editor

diff --git a/Engine/Plugins/Editor/WaveformEditor/WaveformEditor.uplugin b/Engine/Plugins/Editor/WaveformEditor/WaveformEditor.uplugin
--- a/Engine/Plugins/Editor/WaveformEditor/WaveformEditor.uplugin
+++ b/Engine/Plugins/Editor/WaveformEditor/WaveformEditor.uplugin
@@ -1,7 +1,7 @@
 {
 	"FileVersion": 3,
-	"Version": 0,
-	"VersionName": "0.0",
+	"Version": 0.1,
+	"VersionName": "0.1",
 	"FriendlyName": "Waveform Editor",
 	"Description": "Editor tool for waveforms",
 	"Category": "Audio",
