[Waveform Editor] Show Dialog when editor cannot be opened. Impede opening when SoundWave has 0 samples
fix for

diff --git a/Engine/Plugins/Editor/WaveformEditor/Source/WaveformEditor/Private/WaveformEditorInstantiator.cpp b/Engine/Plugins/Editor/WaveformEditor/Source/WaveformEditor/Private/WaveformEditorInstantiator.cpp
--- a/Engine/Plugins/Editor/WaveformEditor/Source/WaveformEditor/Private/WaveformEditorInstantiator.cpp
+++ b/Engine/Plugins/Editor/WaveformEditor/Source/WaveformEditor/Private/WaveformEditorInstantiator.cpp
@@ -5,6 +5,7 @@
 #include "ContentBrowserModule.h"
 #include "Framework/MultiBox/MultiBoxBuilder.h"
 #include "Sound/SoundWave.h"
+#include "SWaveformEditorMessageDialog.h"
 #include "WaveformEditor.h"
 #include "WaveformEditorLog.h"
 
@@ -62,24 +63,66 @@ void FWaveformEditorInstantiator::AddWaveformEditorMenuEntry(FMenuBuilder& MenuB
 
 bool FWaveformEditorInstantiator::CanSoundWaveBeOpenedInEditor(const USoundWave* SoundWaveToEdit)
 {
+	bool bCanOpenWaveEditor = true;
+	
+	FText SoundWaveNameText = FText::FromString(*(SoundWaveToEdit->GetName()));
+	FText ErrorText = FText::Format(LOCTEXT("WaveformEditorOpeningError", "Could not open waveform editor for SoundWave '{0}'"), SoundWaveNameText);
+
 	if (SoundWaveToEdit == nullptr)
 	{
-		UE_LOG(LogWaveformEditor, Warning, TEXT("Could not open waveform editor from null SoundWave"))
-		return false;
+		ErrorText = FText::Format(LOCTEXT("WaveformEditorOpeningError", "Could not open waveform editor for null SoundWave '{0}'"), SoundWaveNameText);
+		bCanOpenWaveEditor = false;
 	}
 	
 	if (SoundWaveToEdit->GetDuration() == 0.f)
 	{
-		UE_LOG(LogWaveformEditor, Warning, TEXT("Could not open waveform editor for soundwave %s, duration is 0"), *(SoundWaveToEdit->GetName()))
-		return false;
+		ErrorText = FText::Format(LOCTEXT("WaveformEditorOpeningError", "Could not open waveform editor for SoundWave '{0}': duration is 0"), SoundWaveNameText);
+		bCanOpenWaveEditor = false;
+	}
+	
+	if (SoundWaveToEdit->NumChannels == 0)
+	{
+		ErrorText = FText::Format(LOCTEXT("WaveformEditorOpeningError", "Could not open waveform editor for SoundWave '{0}': channel count is 0"), SoundWaveNameText);
+		bCanOpenWaveEditor = false;
+	}
+
+	if (SoundWaveToEdit->TotalSamples == 0)
+	{
+		ErrorText = FText::Format(LOCTEXT("WaveformEditorOpeningError", "Could not open waveform editor for SoundWave '{0}': found 0 total samples.\n\nConsider reimporting the asset to fix it."), SoundWaveNameText);
+		bCanOpenWaveEditor = false;
+	}
+
+	if (!bCanOpenWaveEditor)
+	{
+		DisplayErrorDialog(ErrorText);
+	}
+	
+	return bCanOpenWaveEditor;
+}
+
+void FWaveformEditorInstantiator::DisplayErrorDialog(const FText& ErrorMessage) const
+{
+	UE_LOG(LogWaveformEditor, Warning, TEXT("%s"), *ErrorMessage.ToString())
+
+	TSharedPtr<SWindow> OpeningErrorWindow = SNew(SWindow)
+		.Title(LOCTEXT("WaveEditorErrorWindowTitle", "Waveform Editor"))
+		.HasCloseButton(true)
+		.SupportsMaximize(false)
+		.SupportsMinimize(false)
+		.SizingRule(ESizingRule::Autosized);
+
+	OpeningErrorWindow->SetContent(SNew(SWaveformEditorMessageDialog).ParentWindow(OpeningErrorWindow).MessageToDisplay(ErrorMessage));
+
+	TSharedPtr<SWindow> RootWindow = FGlobalTabmanager::Get()->GetRootWindow();
+
+	if (RootWindow.IsValid())
+	{
+		FSlateApplication::Get().AddModalWindow(OpeningErrorWindow.ToSharedRef(), RootWindow);
 	}
-	else if (SoundWaveToEdit->NumChannels == 0.f)
+	else
 	{
-		UE_LOG(LogWaveformEditor, Warning, TEXT("Could not open waveform editor for soundwave %s, channel count is 0"), *(SoundWaveToEdit->GetName()))
-		return false;
+		FSlateApplication::Get().AddWindow(OpeningErrorWindow.ToSharedRef());
 	}
-		
-	return true;
 }
 
 void FWaveformEditorInstantiator::CreateWaveformEditor(TArray<USoundWave*> SoundWavesToEdit)
diff --git a/Engine/Plugins/Editor/WaveformEditor/Source/WaveformEditor/Private/WaveformEditorInstantiator.h b/Engine/Plugins/Editor/WaveformEditor/Source/WaveformEditor/Private/WaveformEditorInstantiator.h
--- a/Engine/Plugins/Editor/WaveformEditor/Source/WaveformEditor/Private/WaveformEditorInstantiator.h
+++ b/Engine/Plugins/Editor/WaveformEditor/Source/WaveformEditor/Private/WaveformEditorInstantiator.h
@@ -21,4 +21,6 @@ private:
 	TSharedRef<FExtender> OnExtendContentBrowserAssetSelectionMenu(const TArray<FAssetData>& SelectedAssets);
 	void AddWaveformEditorMenuEntry(FMenuBuilder& MenuBuilder, TArray<FAssetData> SelectedAssets);
 	bool CanSoundWaveBeOpenedInEditor(const USoundWave* SoundWaveToEdit);
+
+	void DisplayErrorDialog(const FText& ErrorMessage) const;
 };
\ No newline at end of file
diff --git a/Engine/Plugins/Editor/WaveformEditor/Source/WaveformEditorWidgets/Private/SWaveformEditorMessageDialog.cpp b/Engine/Plugins/Editor/WaveformEditor/Source/WaveformEditorWidgets/Private/SWaveformEditorMessageDialog.cpp
new file mode 100644
--- /dev/null
+++ b/Engine/Plugins/Editor/WaveformEditor/Source/WaveformEditorWidgets/Private/SWaveformEditorMessageDialog.cpp
@@ -0,0 +1,72 @@
+// Copyright Epic Games, Inc. All Rights Reserved.
+
+#include "SWaveformEditorMessageDialog.h"
+
+#include "Widgets/Input/SButton.h"
+#include "Widgets/Layout/SBox.h"
+#include "Widgets/SWindow.h"
+#include "Widgets/Text/STextBlock.h"
+
+#define LOCTEXT_NAMESPACE "SWaveformEditorMessageDialog"
+
+void SWaveformEditorMessageDialog::Construct(const FArguments& InArgs)
+{
+	ParentWindowPtr = InArgs._ParentWindow;
+	FText DisplayMessage = InArgs._MessageToDisplay;
+
+	ChildSlot
+	[
+		SNew(SBox)
+		[
+			SNew(SBorder)
+			.Padding(0)
+			.BorderImage(FAppStyle::Get().GetBrush("Brushes.Panel"))
+			[
+				SNew(SVerticalBox)
+				+ SVerticalBox::Slot()
+				.AutoHeight()
+				.Padding(16.0f, 16.0f, 16.0f, 0.0f)
+				[
+					SNew(STextBlock)
+					.WrapTextAt(450.0f)
+					.Text(DisplayMessage)
+					.Justification(ETextJustify::Center)
+				]
+				+ SVerticalBox::Slot()
+				.VAlign(VAlign_Bottom)
+				.Padding(8.0f, 0.f, 8.0f, 16.0f)
+				[
+					SNew(SHorizontalBox)
+					+ SHorizontalBox::Slot()
+					.HAlign(HAlign_Center)
+					[
+						SNew(SButton)
+						.VAlign(VAlign_Center)
+						.TextStyle(FAppStyle::Get(), "DialogButtonText")
+						.Text(LOCTEXT("CloseWaveformEditorErrorMessageDialog", "Close"))
+						.OnClicked(this, &SWaveformEditorMessageDialog::OnCloseButtonPressed)
+						.IsEnabled(this, &SWaveformEditorMessageDialog::CanPressCloseButton)
+					]
+				]
+			]
+		]
+	];
+}
+
+
+bool SWaveformEditorMessageDialog::CanPressCloseButton() const 
+{
+	return ParentWindowPtr.IsValid();
+}
+
+FReply SWaveformEditorMessageDialog::OnCloseButtonPressed() const
+{
+	if (ParentWindowPtr.IsValid())
+	{
+		ParentWindowPtr.Pin()->RequestDestroyWindow();
+	}
+
+	return FReply::Handled();
+}
+
+#undef LOCTEXT_NAMESPACE
\ No newline at end of file
diff --git a/Engine/Plugins/Editor/WaveformEditor/Source/WaveformEditorWidgets/Public/SWaveformEditorMessageDialog.h b/Engine/Plugins/Editor/WaveformEditor/Source/WaveformEditorWidgets/Public/SWaveformEditorMessageDialog.h
new file mode 100644
--- /dev/null
+++ b/Engine/Plugins/Editor/WaveformEditor/Source/WaveformEditorWidgets/Public/SWaveformEditorMessageDialog.h
@@ -0,0 +1,32 @@
+// Copyright Epic Games, Inc. All Rights Reserved.
+
+#pragma once
+
+#include "Widgets/SCompoundWidget.h"
+
+class SWindow;
+
+class WAVEFORMEDITORWIDGETS_API SWaveformEditorMessageDialog : public SCompoundWidget
+{
+public:
+
+	SLATE_BEGIN_ARGS(SWaveformEditorMessageDialog) {}
+
+	/** A reference to the parent window */
+	SLATE_ARGUMENT(TSharedPtr<SWindow>, ParentWindow)
+
+	/** The Message To Display */
+	SLATE_ARGUMENT(FText, MessageToDisplay)
+
+	SLATE_END_ARGS()
+
+	void Construct(const FArguments& InArgs);
+
+private: 
+	/** The parent window of this widget */
+	TWeakPtr<SWindow> ParentWindowPtr;
+
+	bool CanPressCloseButton() const;
+	
+	FReply OnCloseButtonPressed() const;
+};
\ No newline at end of file
