Fix compiler errors caused by wave editor instantiator
- dereferencing a null ptr (warning  C6011)
- wrong mac include for FSlateApplication

diff --git a/Engine/Plugins/Editor/WaveformEditor/Source/WaveformEditor/Private/WaveformEditorInstantiator.cpp b/Engine/Plugins/Editor/WaveformEditor/Source/WaveformEditor/Private/WaveformEditorInstantiator.cpp
--- a/Engine/Plugins/Editor/WaveformEditor/Source/WaveformEditor/Private/WaveformEditorInstantiator.cpp
+++ b/Engine/Plugins/Editor/WaveformEditor/Source/WaveformEditor/Private/WaveformEditorInstantiator.cpp
@@ -66,33 +66,37 @@ bool FWaveformEditorInstantiator::CanSoundWaveBeOpenedInEditor(const USoundWave*
 {
 	bool bCanOpenWaveEditor = true;
 	
-	FText SoundWaveNameText = FText::FromString(*(SoundWaveToEdit->GetName()));
-	FText ErrorText = FText::Format(LOCTEXT("WaveformEditorOpeningError", "Could not open waveform editor for SoundWave '{0}'"), SoundWaveNameText);
+	FText ErrorText = LOCTEXT("WaveformEditorOpeningError", "Could not open waveform editor for Selected SoundWave");
 
 	if (SoundWaveToEdit == nullptr)
 	{
-		ErrorText = FText::Format(LOCTEXT("WaveformEditorOpeningError", "Could not open waveform editor for null SoundWave '{0}'"), SoundWaveNameText);
+		ErrorText = LOCTEXT("WaveformEditorOpeningError", "Could not open waveform editor. Selected SoundWave was null.");
 		bCanOpenWaveEditor = false;
 	}
-	
-	if (SoundWaveToEdit->GetDuration() == 0.f)
-	{
-		ErrorText = FText::Format(LOCTEXT("WaveformEditorOpeningError", "Could not open waveform editor for SoundWave '{0}': duration is 0"), SoundWaveNameText);
-		bCanOpenWaveEditor = false;
-	}
-	
-	if (SoundWaveToEdit->NumChannels == 0)
+	else
 	{
-		ErrorText = FText::Format(LOCTEXT("WaveformEditorOpeningError", "Could not open waveform editor for SoundWave '{0}': channel count is 0"), SoundWaveNameText);
-		bCanOpenWaveEditor = false;
-	}
+		FText SoundWaveNameText = FText::FromString(*(SoundWaveToEdit->GetName()));
+		
+		if (SoundWaveToEdit->GetDuration() == 0.f)
+		{
+			ErrorText = FText::Format(LOCTEXT("WaveformEditorOpeningError", "Could not open waveform editor for SoundWave '{0}': duration is 0"), SoundWaveNameText);
+			bCanOpenWaveEditor = false;
+		}
 
-	if (SoundWaveToEdit->TotalSamples == 0)
-	{
-		ErrorText = FText::Format(LOCTEXT("WaveformEditorOpeningError", "Could not open waveform editor for SoundWave '{0}': found 0 total samples.\n\nConsider reimporting the asset to fix it."), SoundWaveNameText);
-		bCanOpenWaveEditor = false;
+		if (SoundWaveToEdit->NumChannels == 0)
+		{
+			ErrorText = FText::Format(LOCTEXT("WaveformEditorOpeningError", "Could not open waveform editor for SoundWave '{0}': channel count is 0"), SoundWaveNameText);
+			bCanOpenWaveEditor = false;
+		}
+
+		if (SoundWaveToEdit->TotalSamples == 0)
+		{
+			ErrorText = FText::Format(LOCTEXT("WaveformEditorOpeningError", "Could not open waveform editor for SoundWave '{0}': found 0 total samples.\n\nConsider reimporting the asset to fix it."), SoundWaveNameText);
+			bCanOpenWaveEditor = false;
+		}
 	}
 
+
 	if (!bCanOpenWaveEditor)
 	{
 		DisplayErrorDialog(ErrorText);
