Improvements to waveform transformation widgets API
- widgets register for the transformation they want to support through a macro (without taking any dependency on the factory)
- widgets use property handles to interact with the transformation properties, to get full support of property notifiers and editor undo

fix for #jira UE-169805

diff --git a/Engine/Plugins/Editor/WaveformEditor/Source/WaveformEditor/Private/IWaveformEditorDetailsProvider.h b/Engine/Plugins/Editor/WaveformEditor/Source/WaveformEditor/Private/IWaveformEditorDetailsProvider.h
new file mode 100644
--- /dev/null
+++ b/Engine/Plugins/Editor/WaveformEditor/Source/WaveformEditor/Private/IWaveformEditorDetailsProvider.h
@@ -0,0 +1,14 @@
+// Copyright Epic Games, Inc. All Rights Reserved.
+
+#pragma once
+
+#include "Containers/Array.h"
+#include "UObject/ObjectPtr.h"
+
+class IPropertyHandle;
+
+class IWaveformEditorDetailsProvider
+{
+public:
+	virtual void GetHandlesForUObjectProperties(const TObjectPtr<UObject> InUObject, TArray<TSharedRef<IPropertyHandle>>& OutPropertyHandles) = 0;
+};
\ No newline at end of file
diff --git a/Engine/Plugins/Editor/WaveformEditor/Source/WaveformEditor/Private/WaveformEditor.cpp b/Engine/Plugins/Editor/WaveformEditor/Source/WaveformEditor/Private/WaveformEditor.cpp
--- a/Engine/Plugins/Editor/WaveformEditor/Source/WaveformEditor/Private/WaveformEditor.cpp
+++ b/Engine/Plugins/Editor/WaveformEditor/Source/WaveformEditor/Private/WaveformEditor.cpp
@@ -428,7 +428,6 @@ void FWaveformEditor::NotifyPostChange(const FPropertyChangedEvent& PropertyChan
 		return;
 	}
 	
-
 	if (PropertyChangedEvent.ChangeType == EPropertyChangeType::Interactive)
 	{
 		bIsInteractingWithTransformations = true;
@@ -536,14 +535,17 @@ bool FWaveformEditor::SetUpDetailsViews()
 	PropertiesDetails->SetObject(SoundWave);
 
 	TransformationsDetails = PropertyModule.CreateDetailView(Args);
-	FOnGetDetailCustomizationInstance TransformationsDetailsCustomizationInstance = FOnGetDetailCustomizationInstance::CreateLambda([]() { 
-			return MakeShared<FWaveformTransformationsDetailsCustomization>(); 
+	FOnGetDetailCustomizationInstance TransformationsDetailsCustomizationInstance = FOnGetDetailCustomizationInstance::CreateLambda([]() {
+			return MakeShared<FWaveformTransformationsDetailsCustomization>();
 		}
 	);
 
 	TransformationsDetails->RegisterInstancedCustomPropertyLayout(SoundWave->GetClass(), TransformationsDetailsCustomizationInstance);
 	TransformationsDetails->SetObject(SoundWave);
 
+	TransformationsPropertiesPropagator = PropertyModule.CreateDetailView(Args);
+	TransformationsPropertiesPropagator->SetObject(SoundWave);
+	
 	return true;
 }
 
@@ -628,12 +630,29 @@ bool FWaveformEditor::SetUpWaveformPanel()
 	
 	TransportCoordinator = MakeShared<FWaveformEditorTransportCoordinator>(RenderData.ToSharedRef());
 	RenderData->OnRenderDataUpdated.AddSP(TransportCoordinator.Get(), &FWaveformEditorTransportCoordinator::HandleRenderDataUpdate);
+	
+	FOnTransformationsPropertiesRequired OnTransformationPropertiesRequired = FOnTransformationsPropertiesRequired::CreateLambda([this](FTransformationsToPropertiesArray& InObjToPropsMap)
+	{
+		TSharedPtr<FWaveformTransformationsDetailsProvider> TransformationsDetailsProvider = MakeShared<FWaveformTransformationsDetailsProvider>();
+
+		FOnGetDetailCustomizationInstance ProviderInstance = FOnGetDetailCustomizationInstance::CreateLambda( [&]() 
+			{
+				return TransformationsDetailsProvider.ToSharedRef();
+			}
+		);
+
+		TransformationsPropertiesPropagator->RegisterInstancedCustomPropertyLayout(SoundWave->GetClass(), ProviderInstance);
+		TransformationsPropertiesPropagator->ForceRefresh();
+
+		for (FTransformationToPropertiesPair& ObjToPropsPair : InObjToPropsMap)
+		{
+			TransformationsDetailsProvider->GetHandlesForUObjectProperties(ObjToPropsPair.Key, ObjToPropsPair.Value);
+		}
 
-	TFunction<void(FPropertyChangedEvent&, FEditPropertyChain*)> PropertyChangeNotifier = [this](FPropertyChangedEvent& PropertyChangedEvent, FEditPropertyChain* PropertyThatChanged) {
-		this->NotifyPostChange(PropertyChangedEvent, PropertyThatChanged); 
-	};
+		TransformationsDetailsProvider.Reset();
+	});
 
-	TransformationsRenderManager = MakeShared<FWaveformTransformationsRenderManager>(SoundWave, RenderData.ToSharedRef(), PropertyChangeNotifier);
+	TransformationsRenderManager = MakeShared<FWaveformTransformationsRenderManager>(SoundWave, OnTransformationPropertiesRequired);
 	TransformationsRenderManager->OnRenderDataGenerated.AddSP(RenderData.Get(), &FWaveformEditorRenderData::UpdateRenderData);
 
 	TSharedPtr<SWaveformTransformationsOverlay> TransformationsOverlay = SNew(SWaveformTransformationsOverlay, TransformationsRenderManager->GetTransformLayers(), TransportCoordinator.ToSharedRef());
diff --git a/Engine/Plugins/Editor/WaveformEditor/Source/WaveformEditor/Private/WaveformEditorDetailsCustomization.cpp b/Engine/Plugins/Editor/WaveformEditor/Source/WaveformEditor/Private/WaveformEditorDetailsCustomization.cpp
--- a/Engine/Plugins/Editor/WaveformEditor/Source/WaveformEditor/Private/WaveformEditorDetailsCustomization.cpp
+++ b/Engine/Plugins/Editor/WaveformEditor/Source/WaveformEditor/Private/WaveformEditorDetailsCustomization.cpp
@@ -6,22 +6,93 @@
 #include "DetailLayoutBuilder.h"
 #include "Sound/SoundWave.h"
 
-static const FLazyName TransformationsCategoryName("Waveform Processing");
+namespace WaveformTransformationsDetails
+{
+	static const FLazyName TransformationsCategoryName("Waveform Processing");
+
+	void GetUClassProperties(const UStruct* InClass, TArray<FProperty*>& OutProperties)
+	{
+		OutProperties.Empty();
+
+		for (TFieldIterator<FProperty> PropertyIterator(InClass); PropertyIterator; ++PropertyIterator)
+		{
+			OutProperties.Add(*PropertyIterator);
+		}
+	}
+
+	void BuildDetailsView(IDetailLayoutBuilder& DetailLayout)
+	{
+		TArray<FName> CategoryNames;
+		DetailLayout.GetCategoryNames(CategoryNames);
+
+		for (FName& CategoryName : CategoryNames)
+		{
+			if (CategoryName != WaveformTransformationsDetails::TransformationsCategoryName)
+			{
+				DetailLayout.HideCategory(CategoryName);
+			}
+		}
+
+		IDetailCategoryBuilder& CategoryBuilder = DetailLayout.EditCategory("Waveform Processing");
+		CategoryBuilder.InitiallyCollapsed(false);
+		CategoryBuilder.RestoreExpansionState(true);
+	}
+}
 
 void FWaveformTransformationsDetailsCustomization::CustomizeDetails(IDetailLayoutBuilder& DetailLayout)
 {	
-	TArray<FName> CategoryNames;
-	DetailLayout.GetCategoryNames(CategoryNames);
+	WaveformTransformationsDetails::BuildDetailsView(DetailLayout);
+}
+
+void FWaveformTransformationsDetailsProvider::CustomizeDetails(IDetailLayoutBuilder& DetailLayout)
+{
+	WaveformTransformationsDetails::BuildDetailsView(DetailLayout);
+	CachedTransformationsHandle = DetailLayout.GetProperty("Transformations");
+}
+
+void FWaveformTransformationsDetailsProvider::GetHandlesForUObjectProperties(const TObjectPtr<UObject> InUObject, TArray<TSharedRef<IPropertyHandle>>& OutPropertyHandles)
+{
+	check(CachedTransformationsHandle->IsValidHandle());
 
-	for (FName& CategoryName : CategoryNames)
+	OutPropertyHandles.Empty();
+
+	TArray<FProperty*> ObjProperties;
+	TArray<const FName> PropertiesPaths;
+
+
+	uint32 NumTransformations = 0;
+
+	CachedTransformationsHandle->GetNumChildren(NumTransformations);
+	TSharedPtr<IPropertyHandle> TargetHandle;
+
+	for (uint32 TransformIndex = 0; TransformIndex < NumTransformations; ++TransformIndex)
 	{
-		if (CategoryName != TransformationsCategoryName)
+		TSharedPtr<IPropertyHandle> ChildHandle = CachedTransformationsHandle->GetChildHandle(TransformIndex);
+
+		UObject* UObjPtr = nullptr;
+		ChildHandle->GetValue(UObjPtr);
+		
+		if (UObjPtr && UObjPtr == InUObject.Get())
 		{
-			DetailLayout.HideCategory(CategoryName);
+			TargetHandle = ChildHandle;
+			break;
 		}
+
 	}
 
-	IDetailCategoryBuilder& CategoryBuilder = DetailLayout.EditCategory("Waveform Processing");
-	CategoryBuilder.InitiallyCollapsed(false);
-	CategoryBuilder.RestoreExpansionState(true);
-}
\ No newline at end of file
+	if (TargetHandle == nullptr)
+	{
+		return;
+	}
+		
+	WaveformTransformationsDetails::GetUClassProperties(InUObject->GetClass(), ObjProperties);
+
+	for (const FProperty* Property : ObjProperties)
+	{
+		TSharedPtr<IPropertyHandle> PropertyHandle = TargetHandle->GetChildHandle(Property->GetFName());
+		if (PropertyHandle->IsValidHandle())
+		{
+			OutPropertyHandles.Add(PropertyHandle.ToSharedRef());
+		}
+	}
+}
diff --git a/Engine/Plugins/Editor/WaveformEditor/Source/WaveformEditor/Private/WaveformEditorDetailsCustomization.h b/Engine/Plugins/Editor/WaveformEditor/Source/WaveformEditor/Private/WaveformEditorDetailsCustomization.h
--- a/Engine/Plugins/Editor/WaveformEditor/Source/WaveformEditor/Private/WaveformEditorDetailsCustomization.h
+++ b/Engine/Plugins/Editor/WaveformEditor/Source/WaveformEditor/Private/WaveformEditorDetailsCustomization.h
@@ -3,6 +3,7 @@
 #pragma once
 
 #include "IDetailCustomization.h"
+#include "IWaveformEditorDetailsProvider.h"
 #include "UObject/WeakObjectPtr.h"
 
 class IDetailLayoutBuilder;
@@ -10,8 +11,21 @@ class IDetailLayoutBuilder;
 class FWaveformTransformationsDetailsCustomization : public IDetailCustomization
 {
 public:	
+	/** IDetailCustomization interface */
 	void CustomizeDetails(IDetailLayoutBuilder& DetailLayout) override;
+};
+
+
+class FWaveformTransformationsDetailsProvider : public IDetailCustomization, public IWaveformEditorDetailsProvider
+{
+public:
+	/** IDetailCustomization interface */
+	void CustomizeDetails(IDetailLayoutBuilder& DetailLayout) override;
+
+
+	/** IWaveformEditorDetailsProvider interface */
+	virtual void GetHandlesForUObjectProperties(const TObjectPtr<UObject> InUObject, TArray<TSharedRef<IPropertyHandle>>& OutPropertyHandles) override;
 
 private:
-	TWeakObjectPtr<UObject> SoundWaveObject;
+	TSharedPtr<class IPropertyHandle> CachedTransformationsHandle;
 };
\ No newline at end of file
diff --git a/Engine/Plugins/Editor/WaveformEditor/Source/WaveformEditor/Public/WaveformEditor.h b/Engine/Plugins/Editor/WaveformEditor/Source/WaveformEditor/Public/WaveformEditor.h
--- a/Engine/Plugins/Editor/WaveformEditor/Source/WaveformEditor/Public/WaveformEditor.h
+++ b/Engine/Plugins/Editor/WaveformEditor/Source/WaveformEditor/Public/WaveformEditor.h
@@ -114,6 +114,9 @@ private:
 	/** Transformations tab */
 	TSharedPtr<IDetailsView> TransformationsDetails;
 
+	/** Hidden details view to propagate property handles to transformation layers */
+	TSharedPtr<IDetailsView> TransformationsPropertiesPropagator;
+
 	/** Settings Editor App Identifier */
 	static const FName AppIdentifier;
 
diff --git a/Engine/Plugins/Editor/WaveformEditor/Source/WaveformEditorWidgets/Private/WaveformEditorWidgetsModule.cpp b/Engine/Plugins/Editor/WaveformEditor/Source/WaveformEditorWidgets/Private/WaveformEditorWidgetsModule.cpp
--- a/Engine/Plugins/Editor/WaveformEditor/Source/WaveformEditorWidgets/Private/WaveformEditorWidgetsModule.cpp
+++ b/Engine/Plugins/Editor/WaveformEditor/Source/WaveformEditorWidgets/Private/WaveformEditorWidgetsModule.cpp
@@ -4,11 +4,13 @@
 
 #include "Modules/ModuleManager.h"
 #include "WaveformEditorStyle.h"
+#include "WaveformTransformationRendererMapper.h"
 
 void FWaveformEditorWidgetsModule::StartupModule()
 {
 	// Initialize static style instance 
 	FWaveformEditorStyle::Init();
+	FWaveformTransformationRendererMapper::Init();
 }
 
 void FWaveformEditorWidgetsModule::ShutdownModule()
diff --git a/Engine/Plugins/Editor/WaveformEditor/Source/WaveformEditorWidgets/Private/WaveformTransformationRenderLayerFactory.cpp b/Engine/Plugins/Editor/WaveformEditor/Source/WaveformEditorWidgets/Private/WaveformTransformationRenderLayerFactory.cpp
new file mode 100644
--- /dev/null
+++ b/Engine/Plugins/Editor/WaveformEditor/Source/WaveformEditorWidgets/Private/WaveformTransformationRenderLayerFactory.cpp
@@ -0,0 +1,35 @@
+// Copyright Epic Games, Inc. All Rights Reserved.
+
+#include "WaveformTransformationRenderLayerFactory.h"
+
+#include "IWaveformTransformation.h"
+#include "IWaveformTransformationRenderer.h"
+#include "WaveformTransformationRendererMapper.h"
+
+TSharedPtr<IWaveformTransformationRenderer> FWaveformTransformationRenderLayerFactory::Create(TObjectPtr<UWaveformTransformationBase> InTransformationToRender)
+{
+	TSharedPtr<IWaveformTransformationRenderer> OutRenderer = nullptr;
+	UClass* TransformationClass = InTransformationToRender->GetClass();
+
+	FWaveformTransformRendererInstantiator* Instantiator = FWaveformTransformationRendererMapper::Get().GetRenderer(InTransformationToRender->GetClass());
+
+	if (Instantiator)
+	{
+		Instantiator->CheckCallable();
+		OutRenderer = (*Instantiator)();
+	}
+
+	return OutRenderer;
+}
+
+TSharedPtr<IWaveformTransformationRenderer> FWaveformTransformationRenderLayerFactory::Create(TObjectPtr<UWaveformTransformationBase> InTransformationToRender, TArray<TSharedRef<IPropertyHandle>>& TransformationsProperties)
+{
+	TSharedPtr<IWaveformTransformationRenderer> OutRenderer = Create(InTransformationToRender);
+
+	if (OutRenderer)
+	{
+		OutRenderer->SetPropertyHandles(TransformationsProperties);
+	}
+
+	return OutRenderer;
+}
diff --git a/Engine/Plugins/Editor/WaveformEditor/Source/WaveformEditorWidgets/Private/WaveformTransformationRendererMapper.cpp b/Engine/Plugins/Editor/WaveformEditor/Source/WaveformEditorWidgets/Private/WaveformTransformationRendererMapper.cpp
new file mode 100644
--- /dev/null
+++ b/Engine/Plugins/Editor/WaveformEditor/Source/WaveformEditorWidgets/Private/WaveformTransformationRendererMapper.cpp
@@ -0,0 +1,31 @@
+// Copyright Epic Games, Inc. All Rights Reserved.
+
+#include "WaveformTransformationRendererMapper.h"
+
+FWaveformTransformationRendererMapper& FWaveformTransformationRendererMapper::Get()
+{
+	check(Instance);
+	return *Instance;
+}
+
+void FWaveformTransformationRendererMapper::Init()
+{
+	if (Instance == nullptr)
+	{
+		Instance = MakeUnique<FWaveformTransformationRendererMapper>();
+	}
+}
+
+void FWaveformTransformationRendererMapper::UnregisterRenderer(const UClass* SupportedTransformation)
+{
+	check(Instance)
+	Instance->TransformationsTypeMap.FindAndRemoveChecked(SupportedTransformation);
+}
+
+FWaveformTransformRendererInstantiator* FWaveformTransformationRendererMapper::GetRenderer(const UClass* WaveformTransformationClass)
+{
+	check(Instance)
+	return Instance->TransformationsTypeMap.Find(WaveformTransformationClass);
+}
+
+TUniquePtr<FWaveformTransformationRendererMapper> FWaveformTransformationRendererMapper::Instance;
diff --git a/Engine/Plugins/Editor/WaveformEditor/Source/WaveformEditorWidgets/Public/IWaveformTransformationRenderer.h b/Engine/Plugins/Editor/WaveformEditor/Source/WaveformEditorWidgets/Public/IWaveformTransformationRenderer.h
--- a/Engine/Plugins/Editor/WaveformEditor/Source/WaveformEditorWidgets/Public/IWaveformTransformationRenderer.h
+++ b/Engine/Plugins/Editor/WaveformEditor/Source/WaveformEditorWidgets/Public/IWaveformTransformationRenderer.h
@@ -2,12 +2,10 @@
 
 #pragma once
 
-#include "HAL/Platform.h"
 #include "IWaveformTransformation.h"
-#include "Sound/SoundWave.h"
 #include "Templates/Function.h"
-#include "UObject/UnrealType.h"
 
+class IPropertyHandle;
 class FCursorReply;
 class FPaintArgs;
 class FReply;
@@ -15,6 +13,7 @@ class FSlateRect;
 class FSlateWindowElementList;
 class FWidgetStyle;
 class SWidget;
+class UClass;
 struct FGeometry;
 struct FPointerEvent;
 
@@ -44,48 +43,10 @@ public:
 	virtual FCursorReply OnCursorQuery(const FGeometry& MyGeometry, const FPointerEvent& CursorEvent) const = 0;
 
 	/* Gets info about the transformation being rendered */
-	void SetTransformationWaveInfo(FWaveformTransformationRenderInfo InWaveInfo)
-	{
-		TransformationWaveInfo = InWaveInfo;
-	}
-
-	/* Sets a callback to notify when a property has been changed by the renderer */
-	void SetTransformationNotifier(const TFunction<void(FPropertyChangedEvent&, FEditPropertyChain*)>& InTransformationChangeNotifier)
-	{
-		TransformationChangeNotifier = InTransformationChangeNotifier;
-	}
-
-protected:
-
-	/* Notifies that a property has been changed through TransformationChangeNotifier  */
-	void NotifyTransformationPropertyChanged(const TObjectPtr<UWaveformTransformationBase> EditedTransformation, const FName& PropertyName, const EPropertyChangeType::Type InChangeType)
-	{
-		FProperty* Property = EditedTransformation->GetClass()->FindPropertyByName(PropertyName);
-		check(Property);
-
-		USoundWave* ParentSoundWave = EditedTransformation->GetTypedOuter<USoundWave>();
-		check(ParentSoundWave)
-
-		FProperty* TransformationsProperty = ParentSoundWave->GetClass()->FindPropertyByName(GET_MEMBER_NAME_CHECKED(USoundWave, Transformations));
-		check(TransformationsProperty)
-
-		FEditPropertyChain EditChain;
-		EditChain.AddHead(TransformationsProperty);
-		EditChain.AddTail(Property);
-		EditChain.SetActivePropertyNode(Property);
-
-		FPropertyChangedEvent PropertyEvent(Property, InChangeType);
-		FPropertyChangedEvent TransformationsPropertyEvent(TransformationsProperty, InChangeType);
-
-		EditedTransformation->PostEditChangeProperty(PropertyEvent);
-		ParentSoundWave->PostEditChangeProperty(TransformationsPropertyEvent);
-
-		TransformationChangeNotifier(PropertyEvent, &EditChain);
-	}
-
-	FWaveformTransformationRenderInfo TransformationWaveInfo;
+	virtual void SetTransformationWaveInfo(const FWaveformTransformationRenderInfo& InWaveInfo) = 0;
+	
+	/* Gets the transformations properties. IPropertyHandle is used to get/set properties to have full support of change notifiers, editor undo ecc  */
+	virtual void SetPropertyHandles(const TArray<TSharedRef<IPropertyHandle>>& InPropertyHandles) = 0;
+};
 
-private:
-	TFunction<void(FPropertyChangedEvent&, FEditPropertyChain*)> TransformationChangeNotifier;
 
-};
\ No newline at end of file
diff --git a/Engine/Plugins/Editor/WaveformEditor/Source/WaveformEditorWidgets/Public/WaveformTransformationRenderLayerFactory.h b/Engine/Plugins/Editor/WaveformEditor/Source/WaveformEditorWidgets/Public/WaveformTransformationRenderLayerFactory.h
new file mode 100644
--- /dev/null
+++ b/Engine/Plugins/Editor/WaveformEditor/Source/WaveformEditorWidgets/Public/WaveformTransformationRenderLayerFactory.h
@@ -0,0 +1,20 @@
+// Copyright Epic Games, Inc. All Rights Reserved.
+
+#pragma once
+
+#include "UObject/ObjectPtr.h"
+#include "Templates/SharedPointer.h"
+
+class IWaveformTransformationRenderer;
+class UWaveformTransformationBase;
+class IPropertyHandle;
+
+class WAVEFORMEDITORWIDGETS_API FWaveformTransformationRenderLayerFactory
+{
+public:
+	~FWaveformTransformationRenderLayerFactory() = default;
+
+	TSharedPtr<IWaveformTransformationRenderer> Create(TObjectPtr<UWaveformTransformationBase> InTransformationToRender);
+	TSharedPtr<IWaveformTransformationRenderer> Create(TObjectPtr<UWaveformTransformationBase> InTransformationToRender, TArray<TSharedRef<IPropertyHandle>>& TransformationsProperties);
+
+};
\ No newline at end of file
diff --git a/Engine/Plugins/Editor/WaveformEditor/Source/WaveformEditorWidgets/Public/WaveformTransformationRendererMapper.h b/Engine/Plugins/Editor/WaveformEditor/Source/WaveformEditorWidgets/Public/WaveformTransformationRendererMapper.h
new file mode 100644
--- /dev/null
+++ b/Engine/Plugins/Editor/WaveformEditor/Source/WaveformEditorWidgets/Public/WaveformTransformationRendererMapper.h
@@ -0,0 +1,37 @@
+// Copyright Epic Games, Inc. All Rights Reserved.
+
+#pragma once
+
+#include "IWaveformTransformationRenderer.h"
+
+using FWaveformTransformRendererInstantiator = TFunction<TSharedPtr<IWaveformTransformationRenderer>()>;
+
+class WAVEFORMEDITORWIDGETS_API FWaveformTransformationRendererMapper
+{
+public:
+	/** Access the singleton instance for mapper */
+	static FWaveformTransformationRendererMapper& Get();
+	static void Init();
+
+	template<typename ConcreteRendererType>
+	bool RegisterRenderer(const UClass* SupportedTransformation)
+	{
+		FWaveformTransformRendererInstantiator Instantiator = []()
+		{
+			return MakeShared<ConcreteRendererType>();
+		};
+
+		check(Instance)
+		Instance->TransformationsTypeMap.Add(SupportedTransformation, Instantiator);
+		return true;
+	}
+
+	void UnregisterRenderer(const UClass* SupportedTransformation);
+
+	FWaveformTransformRendererInstantiator* GetRenderer(const UClass* WaveformTransformationClass);
+
+private:
+	static TUniquePtr<FWaveformTransformationRendererMapper> Instance;
+
+	TMap<const UClass*, FWaveformTransformRendererInstantiator> TransformationsTypeMap;
+};
\ No newline at end of file
diff --git a/Engine/Plugins/Editor/WaveformEditor/Source/WaveformTransformationsWidgets/Private/WaveformTransformationDurationRenderer.h b/Engine/Plugins/Editor/WaveformEditor/Source/WaveformTransformationsWidgets/Private/WaveformTransformationDurationRenderer.h
--- a/Engine/Plugins/Editor/WaveformEditor/Source/WaveformTransformationsWidgets/Private/WaveformTransformationDurationRenderer.h
+++ b/Engine/Plugins/Editor/WaveformEditor/Source/WaveformTransformationsWidgets/Private/WaveformTransformationDurationRenderer.h
@@ -3,7 +3,7 @@
 #pragma once
 
 #include "WaveformTransformationRendererBase.h"
-#include "WaveformEditorRenderData.h"
+
 
 namespace WaveformTransformationDurationHiglightParams
 {
@@ -14,16 +14,16 @@ namespace WaveformTransformationDurationHiglightParams
 class FWaveformTransformationDurationRenderer : public FWaveformTransformationRendererBase
 {
 public:
-	explicit FWaveformTransformationDurationRenderer(const TSharedRef<FWaveformEditorRenderData> InRenderData)
+	explicit FWaveformTransformationDurationRenderer(const uint32 InOriginalWaveformFrames)
+		: OriginalWaveformFrames(InOriginalWaveformFrames)
 	{
-		WaveformRenderData = InRenderData;
 	}
 
 	virtual int32 OnPaint(const FPaintArgs& Args, const FGeometry& AllottedGeometry, const FSlateRect& MyCullingRect, FSlateWindowElementList& OutDrawElements, int32 LayerId, const FWidgetStyle& InWidgetStyle, bool bParentEnabled) const override
 	{
-		const float StartTimeRatio = (float)TransformationWaveInfo.StartFrameOffset / WaveformRenderData->GetOriginalWaveformFrames() / TransformationWaveInfo.NumChannels;
+		const float StartTimeRatio = (float)TransformationWaveInfo.StartFrameOffset / OriginalWaveformFrames / TransformationWaveInfo.NumChannels;
 		const uint32 EndSample = TransformationWaveInfo.StartFrameOffset + TransformationWaveInfo.NumAvilableSamples;
-		const float EndTimeRatio = (float)EndSample / WaveformRenderData->GetOriginalWaveformFrames() / TransformationWaveInfo.NumChannels;
+		const float EndTimeRatio = (float)EndSample / OriginalWaveformFrames / TransformationWaveInfo.NumChannels;
 
 		const bool bRenderLeftBox = StartTimeRatio >= 0.f;
 		const bool bRenderRightBox = EndTimeRatio <= 1.f;
@@ -58,6 +58,8 @@ public:
 
 		return LayerId;
 	}
+
+	void SetOriginalWaveformFrames(const uint32 NumFrames) { OriginalWaveformFrames = NumFrames; }
 private:
-	TSharedPtr<FWaveformEditorRenderData> WaveformRenderData = nullptr;
+	uint32 OriginalWaveformFrames = 0;
 };
\ No newline at end of file
diff --git a/Engine/Plugins/Editor/WaveformEditor/Source/WaveformTransformationsWidgets/Private/WaveformTransformationRenderLayerFactory.cpp b/Engine/Plugins/Editor/WaveformEditor/Source/WaveformTransformationsWidgets/Private/WaveformTransformationRenderLayerFactory.cpp
deleted file mode 100644
diff --git a/Engine/Plugins/Editor/WaveformEditor/Source/WaveformTransformationsWidgets/Private/WaveformTransformationRendererBase.cpp b/Engine/Plugins/Editor/WaveformEditor/Source/WaveformTransformationsWidgets/Private/WaveformTransformationRendererBase.cpp
--- a/Engine/Plugins/Editor/WaveformEditor/Source/WaveformTransformationsWidgets/Private/WaveformTransformationRendererBase.cpp
+++ b/Engine/Plugins/Editor/WaveformEditor/Source/WaveformTransformationsWidgets/Private/WaveformTransformationRendererBase.cpp
@@ -1,7 +1,12 @@
 // Copyright Epic Games, Inc. All Rights Reserved.
 
 #include "WaveformTransformationRendererBase.h"
+
+#include "DetailLayoutBuilder.h"
+#include "DetailCategoryBuilder.h"
+#include "Editor.h"
 #include "Input/Reply.h"
+#include "Sound/SoundWave.h"
 
 int32 FWaveformTransformationRendererBase::OnPaint(const FPaintArgs& Args, const FGeometry& AllottedGeometry, const FSlateRect& MyCullingRect, FSlateWindowElementList& OutDrawElements, int32 LayerId, const FWidgetStyle& InWidgetStyle, bool bParentEnabled) const
 {
@@ -37,3 +42,44 @@ FCursorReply FWaveformTransformationRendererBase::OnCursorQuery(const FGeometry&
 {
 	return FCursorReply::Unhandled();
 }
+
+void FWaveformTransformationRendererBase::SetTransformationWaveInfo(const FWaveformTransformationRenderInfo& InWaveInfo)
+{
+	TransformationWaveInfo = InWaveInfo;
+}
+
+void FWaveformTransformationRendererBase::SetPropertyHandles(const TArray<TSharedRef<IPropertyHandle>>& InPropertyHandles)
+{
+	CachedPropertyHandles = InPropertyHandles;
+}
+
+TSharedPtr<IPropertyHandle> FWaveformTransformationRendererBase::GetPropertyHandle(const FName& PropertyName) const 
+{
+	TSharedPtr<IPropertyHandle> HandlePtr = nullptr;
+
+	for (const TSharedRef<IPropertyHandle>& PropertyHandle : CachedPropertyHandles)
+	{
+		if (PropertyHandle.Get().GetProperty()->GetFName() == PropertyName)
+		{
+			check(PropertyHandle.Get().IsValidHandle());
+			return PropertyHandle;
+		}
+	}
+	check(HandlePtr)
+	return HandlePtr;
+}
+
+int32 FWaveformTransformationRendererBase::BeginTransaction(const TCHAR* TransactionContext, const FText& Description, UObject* PrimaryObject)
+{
+	if (GEditor && GEditor->Trans)
+	{
+		return GEditor->BeginTransaction(TransactionContext, Description, PrimaryObject);
+	}
+
+	return INDEX_NONE;
+}
+
+int32 FWaveformTransformationRendererBase::EndTransaction()
+{
+	return GEditor->EndTransaction();
+}
diff --git a/Engine/Plugins/Editor/WaveformEditor/Source/WaveformTransformationsWidgets/Private/WaveformTransformationRendererBase.h b/Engine/Plugins/Editor/WaveformEditor/Source/WaveformTransformationsWidgets/Private/WaveformTransformationRendererBase.h
--- a/Engine/Plugins/Editor/WaveformEditor/Source/WaveformTransformationsWidgets/Private/WaveformTransformationRendererBase.h
+++ b/Engine/Plugins/Editor/WaveformEditor/Source/WaveformTransformationsWidgets/Private/WaveformTransformationRendererBase.h
@@ -2,13 +2,15 @@
 
 #pragma once
 
+#include "IDetailCustomization.h"
 #include "IWaveformTransformationRenderer.h"
+#include "PropertyHandle.h"
 
 class FWaveformTransformationRendererBase : public IWaveformTransformationRenderer
 {
 public:
 	FWaveformTransformationRendererBase() = default;
-
+	
 	virtual int32 OnPaint(const FPaintArgs& Args, const FGeometry& AllottedGeometry, const FSlateRect& MyCullingRect, FSlateWindowElementList& OutDrawElements, int32 LayerId, const FWidgetStyle& InWidgetStyle, bool bParentEnabled) const override;
 	virtual void Tick(const FGeometry& AllottedGeometry, const double InCurrentTime, const float InDeltaTime) override {};
 	virtual FReply OnMouseButtonDown(SWidget& OwnerWidget, const FGeometry& MyGeometry, const FPointerEvent& MouseEvent) override;
@@ -19,4 +21,35 @@ public:
 	virtual void OnMouseEnter(const FGeometry& MyGeometry, const FPointerEvent& MouseEvent) override {};
 	virtual void OnMouseLeave(const FPointerEvent& MouseEvent) override {};
 	virtual FCursorReply OnCursorQuery(const FGeometry& MyGeometry, const FPointerEvent& CursorEvent) const override;
+
+	virtual void SetTransformationWaveInfo(const FWaveformTransformationRenderInfo& InWaveInfo) override;
+	virtual void SetPropertyHandles(const TArray<TSharedRef<IPropertyHandle>>& InPropertyHandles) override;
+
+protected:
+	TSharedPtr<IPropertyHandle> GetPropertyHandle(const FName& PropertyName) const;
+
+	template<typename T>
+	T GetPropertyValue(const TSharedPtr<IPropertyHandle>& PropertyHandle)
+	{
+		check(PropertyHandle.IsValid());
+		T OutValue;
+		PropertyHandle->GetValue(OutValue);
+		return OutValue;
+	}
+
+	template<typename T>
+	void SetPropertyValue(const FName& PropertyName, const T Value, const EPropertyValueSetFlags::Type Flags = EPropertyValueSetFlags::DefaultFlags)
+	{
+		TSharedPtr<IPropertyHandle> Handle = GetPropertyHandle(PropertyName);
+		Handle->SetValue(Value, Flags);
+	}
+
+
+	int32 BeginTransaction(const TCHAR* TransactionContext, const FText& Description, UObject* PrimaryObject);
+	int32 EndTransaction();
+
+	FWaveformTransformationRenderInfo TransformationWaveInfo;
+
+private:
+	TArray<TSharedRef<IPropertyHandle>> CachedPropertyHandles;	
 };
\ No newline at end of file
diff --git a/Engine/Plugins/Editor/WaveformEditor/Source/WaveformTransformationsWidgets/Private/WaveformTransformationTrimFadeRenderer.cpp b/Engine/Plugins/Editor/WaveformEditor/Source/WaveformTransformationsWidgets/Private/WaveformTransformationTrimFadeRenderer.cpp
--- a/Engine/Plugins/Editor/WaveformEditor/Source/WaveformTransformationsWidgets/Private/WaveformTransformationTrimFadeRenderer.cpp
+++ b/Engine/Plugins/Editor/WaveformEditor/Source/WaveformTransformationsWidgets/Private/WaveformTransformationTrimFadeRenderer.cpp
@@ -1,16 +1,13 @@
 // Copyright Epic Games, Inc. All Rights Reserved.
 
 #include "WaveformTransformationTrimFadeRenderer.h"
+
 #include "Layout/Geometry.h"
+#include "PropertyHandle.h"
 #include "Rendering/DrawElements.h"
+#include "WaveformTransformationTrimFade.h"
 #include "Widgets/SWidget.h"
 
-FWaveformTransformationTrimFadeRenderer::FWaveformTransformationTrimFadeRenderer(const TObjectPtr<UWaveformTransformationTrimFade> TransformationToRender)
-{
-	check(TransformationToRender);
-	TrimFadeTransform = TransformationToRender;
-}
-
 int32 FWaveformTransformationTrimFadeRenderer::OnPaint(const FPaintArgs& Args, const FGeometry& AllottedGeometry, const FSlateRect& MyCullingRect, FSlateWindowElementList& OutDrawElements, int32 LayerId, const FWidgetStyle& InWidgetStyle, bool bParentEnabled) const
 {
 	LayerId = DrawTrimHandles(AllottedGeometry, OutDrawElements, LayerId);
@@ -95,7 +92,7 @@ int32 FWaveformTransformationTrimFadeRenderer::DrawFadeCurves(const FGeometry& A
 
 void FWaveformTransformationTrimFadeRenderer::GenerateFadeCurves(const FGeometry& AllottedGeometry)
 {
-	const float FadeInFrames = TrimFadeTransform->StartFadeTime * TransformationWaveInfo.SampleRate;
+	const float FadeInFrames = GetPropertyValue<float>(StartFadeTimeHandle) * TransformationWaveInfo.SampleRate;
 	const uint32 FadeInPixelLenght = FadeInFrames * PixelsPerFrame;
 	FadeInStartX = FMath::RoundToInt32(StartTimeHandleX);
 	FadeInEndX = FMath::RoundToInt32(FMath::Clamp(StartTimeHandleX + FadeInPixelLenght, StartTimeHandleX, EndTimeHandleX));
@@ -106,13 +103,13 @@ void FWaveformTransformationTrimFadeRenderer::GenerateFadeCurves(const FGeometry
 	for (uint32 Pixel = 0; Pixel < DisplayedFadeInPixelLenght; ++Pixel)
 	{
 		const double FadeFraction = (float)Pixel / FadeInPixelLenght;
-		const double CurveValue = Pixel != FadeInPixelLenght - 1 ? 1.f - FMath::Pow(FadeFraction, TrimFadeTransform->StartFadeCurve) : 0.f;
+		const double CurveValue = Pixel != FadeInPixelLenght - 1 ? 1.f - FMath::Pow(FadeFraction, GetPropertyValue<float>(StartFadeCurveHandle)) : 0.f;
 
 		const uint32 XCoordinate = Pixel + FadeInStartX;
 		FadeInCurvePoints[Pixel] = FVector2D(XCoordinate, CurveValue * AllottedGeometry.Size.Y);
 	}
 
-	const float FadeOutFrames = TrimFadeTransform->EndFadeTime * TransformationWaveInfo.SampleRate;
+	const float FadeOutFrames = GetPropertyValue<float>(EndFadeTimeHandle) * TransformationWaveInfo.SampleRate;
 	const float FadeOutPixelLength = FadeOutFrames * PixelsPerFrame;
 	FadeOutStartX = FMath::RoundToInt32(FMath::Clamp(EndTimeHandleX - FadeOutPixelLength, StartTimeHandleX, EndTimeHandleX));
 	FadeOutEndX = FMath::RoundToInt32(EndTimeHandleX);
@@ -124,7 +121,7 @@ void FWaveformTransformationTrimFadeRenderer::GenerateFadeCurves(const FGeometry
 	for (uint32 Pixel = 0; Pixel < DisplayedFadeOutPixelLength; ++Pixel)
 	{
 		const double FadeFraction = (float)(Pixel + FadeOutPixelOffset) / FadeOutPixelLength;
-		const double CurveValue = Pixel != FadeOutPixelLength - 1 ? FMath::Pow(FadeFraction, TrimFadeTransform->EndFadeCurve) : 1.f;
+		const double CurveValue = Pixel != FadeOutPixelLength - 1 ? FMath::Pow(FadeFraction, GetPropertyValue<float>(EndFadeCurveHandle)) : 1.f;
 
 		const uint32 XCoordinate = Pixel + FadeOutStartX;
 		FadeOutCurvePoints[Pixel] = FVector2D(XCoordinate, CurveValue * AllottedGeometry.Size.Y);
@@ -153,6 +150,18 @@ FCursorReply FWaveformTransformationTrimFadeRenderer::OnCursorQuery(const FGeome
 	return FCursorReply::Unhandled();
 }
 
+void FWaveformTransformationTrimFadeRenderer::SetPropertyHandles(const TArray<TSharedRef<IPropertyHandle>>& InPropertyHandles)
+{
+	FWaveformTransformationRendererBase::SetPropertyHandles(InPropertyHandles);
+
+	StartTimeHandle = GetPropertyHandle(GET_MEMBER_NAME_CHECKED(UWaveformTransformationTrimFade, StartTime));
+	EndTimeHandle  = GetPropertyHandle(GET_MEMBER_NAME_CHECKED(UWaveformTransformationTrimFade, EndTime));
+	StartFadeTimeHandle = GetPropertyHandle(GET_MEMBER_NAME_CHECKED(UWaveformTransformationTrimFade, StartFadeTime));
+	StartFadeCurveHandle = GetPropertyHandle(GET_MEMBER_NAME_CHECKED(UWaveformTransformationTrimFade, StartFadeCurve));
+	EndFadeTimeHandle = GetPropertyHandle(GET_MEMBER_NAME_CHECKED(UWaveformTransformationTrimFade, EndFadeTime));
+	EndFadeCurveHandle = GetPropertyHandle(GET_MEMBER_NAME_CHECKED(UWaveformTransformationTrimFade, EndFadeCurve));
+}
+
 FReply FWaveformTransformationTrimFadeRenderer::OnMouseWheel(SWidget& OwnerWidget, const FGeometry& MyGeometry, const FPointerEvent& MouseEvent)
 {
 	const FVector2D LocalCursorPosition = GetLocalCursorPosition(MouseEvent, MyGeometry);
@@ -160,16 +169,22 @@ FReply FWaveformTransformationTrimFadeRenderer::OnMouseWheel(SWidget& OwnerWidge
 	if (IsCursorInFadeInInteractionRange(LocalCursorPosition, MyGeometry))
 	{		
 		const float FadeCurveDelta = MouseEvent.GetWheelDelta() * MouseWheelStep;
-		TrimFadeTransform->StartFadeCurve = FMath::Clamp(TrimFadeTransform->StartFadeCurve + FadeCurveDelta, 0.f, 10.f);
-		NotifyTransformationPropertyChanged(TrimFadeTransform, GET_MEMBER_NAME_CHECKED(UWaveformTransformationTrimFade, StartFadeCurve), EPropertyChangeType::ValueSet);
+		const float FadeCurveValue = FMath::Clamp(GetPropertyValue<float>(StartFadeCurveHandle) + FadeCurveDelta, 0.f, 10.f);
+
+		EPropertyValueSetFlags::Type PropertyFlags = EPropertyValueSetFlags::InteractiveChange;
+		StartFadeCurveHandle->SetValue(FadeCurveValue, PropertyFlags);
+
 		return FReply::Handled();
 	}
 
 	if (IsCursorInFadeOutInteractionRange(LocalCursorPosition, MyGeometry))
 	{
 		const float FadeCurveDelta = MouseEvent.GetWheelDelta() * MouseWheelStep;
-		TrimFadeTransform->EndFadeCurve = FMath::Clamp(TrimFadeTransform->EndFadeCurve + FadeCurveDelta, 0.f, 10.f);
-		NotifyTransformationPropertyChanged(TrimFadeTransform, GET_MEMBER_NAME_CHECKED(UWaveformTransformationTrimFade, EndFadeCurve), EPropertyChangeType::ValueSet);
+		const float FadeCurveValue = FMath::Clamp(GetPropertyValue<float>(EndFadeCurveHandle) + FadeCurveDelta, 0.f, 10.f);
+		
+		EPropertyValueSetFlags::Type PropertyFlags = EPropertyValueSetFlags::InteractiveChange;
+		EndFadeCurveHandle->SetValue(FadeCurveValue, PropertyFlags);
+		
 		return FReply::Handled();
 	}
 
@@ -182,7 +197,7 @@ FVector2D FWaveformTransformationTrimFadeRenderer::GetLocalCursorPosition(const
 	return  EventGeometry.AbsoluteToLocal(ScreenSpacePosition);
 }
 
-float FWaveformTransformationTrimFadeRenderer::ConvertXRatioToTime(const float InRatio) const
+double FWaveformTransformationTrimFadeRenderer::ConvertXRatioToTime(const float InRatio) const
 {
 	const float NumFrames = TransformationWaveInfo.NumAvilableSamples / TransformationWaveInfo.NumChannels;
 	const float FrameSelected = NumFrames * InRatio;
@@ -191,14 +206,9 @@ float FWaveformTransformationTrimFadeRenderer::ConvertXRatioToTime(const float I
 
 void FWaveformTransformationTrimFadeRenderer::Tick(const FGeometry& AllottedGeometry, const double InCurrentTime, const float InDeltaTime)
 { 
-	if (!TrimFadeTransform)
-	{
-		return;
-	}
-
 	const float NumFrames = TransformationWaveInfo.NumAvilableSamples / TransformationWaveInfo.NumChannels;
-	const double FirstFrame = FMath::Clamp((TrimFadeTransform->StartTime * TransformationWaveInfo.SampleRate) , 0.f, NumFrames);
-	const double EndFrame = FMath::Clamp((TrimFadeTransform->EndTime * TransformationWaveInfo.SampleRate), FirstFrame, NumFrames);
+	const double FirstFrame = FMath::Clamp((GetPropertyValue<double>(StartTimeHandle) * TransformationWaveInfo.SampleRate) , 0.f, NumFrames);
+	const double EndFrame = FMath::Clamp((GetPropertyValue<double>(EndTimeHandle) * TransformationWaveInfo.SampleRate), FirstFrame, NumFrames);
 
 	PixelsPerFrame = FMath::Max(AllottedGeometry.GetLocalSize().X / NumFrames, 0.0);
 
@@ -233,8 +243,10 @@ FReply FWaveformTransformationTrimFadeRenderer::OnMouseButtonDown(SWidget& Owner
 
 		if (TrimFadeInteractionType != ETrimFadeInteractionType::None)
 		{
+			BeginTransaction(TEXT("PropertyEditor"), FText::Format(NSLOCTEXT("PropertyEditor", "SetNumericPropertyTransaction", "Edit {0}"), GetPropertyEditedByCurrentInteraction()), nullptr);
 			return FReply::Handled().CaptureMouse(OwnerWidget.AsShared()).PreventThrottling();
 		}
+
 	}
 
 	return FReply::Unhandled();
@@ -245,7 +257,7 @@ FReply FWaveformTransformationTrimFadeRenderer::OnMouseMove(SWidget& OwnerWidget
 {
 	if (MouseEvent.IsMouseButtonDown(EKeys::LeftMouseButton) && TrimFadeInteractionType != ETrimFadeInteractionType::None)
 	{
-		SetPropertyValueDependingOnInteractionType(MouseEvent, MyGeometry, EPropertyChangeType::Interactive);
+		SetPropertyValueDependingOnInteractionType(MouseEvent, MyGeometry, EPropertyValueSetFlags::InteractiveChange | EPropertyValueSetFlags::NotTransactable);
 
 		return FReply::Handled().CaptureMouse(OwnerWidget.AsShared());
 	}
@@ -257,8 +269,8 @@ FReply FWaveformTransformationTrimFadeRenderer::OnMouseButtonUp(SWidget& OwnerWi
 {
 	if (TrimFadeInteractionType != ETrimFadeInteractionType::None)
 	{
-		SetPropertyValueDependingOnInteractionType(MouseEvent, MyGeometry, EPropertyChangeType::ValueSet);
-
+		SetPropertyValueDependingOnInteractionType(MouseEvent, MyGeometry);
+		EndTransaction();
 		TrimFadeInteractionType = ETrimFadeInteractionType::None;
 		return FReply::Handled().ReleaseMouseCapture();
 	}
@@ -266,38 +278,41 @@ FReply FWaveformTransformationTrimFadeRenderer::OnMouseButtonUp(SWidget& OwnerWi
 	return FReply::Unhandled();
 }
 
-void FWaveformTransformationTrimFadeRenderer::SetPropertyValueDependingOnInteractionType(const FPointerEvent& MouseEvent, const FGeometry& WidgetGeometry, const EPropertyChangeType::Type DesiredChangeType)
+void FWaveformTransformationTrimFadeRenderer::SetPropertyValueDependingOnInteractionType(const FPointerEvent& MouseEvent, const FGeometry& WidgetGeometry, const EPropertyValueSetFlags::Type Flags /*= EPropertyValueSetFlags::DefaultFlags*/)
 {
 	const FVector2D LocalCursorPosition = GetLocalCursorPosition(MouseEvent, WidgetGeometry);
 	const float LocalCursorXRatio = FMath::Clamp(LocalCursorPosition.X / WidgetGeometry.GetLocalSize().X, 0.f, 1.f);
-	const float SelectedTime = ConvertXRatioToTime(LocalCursorXRatio);
+	const double SelectedTime = ConvertXRatioToTime(LocalCursorXRatio);
+
+	TSharedPtr<IPropertyHandle> PropertyHandle;
 
 	switch (TrimFadeInteractionType)
 	{
 	case FWaveformTransformationTrimFadeRenderer::ETrimFadeInteractionType::None:
 		break;
 	case FWaveformTransformationTrimFadeRenderer::ETrimFadeInteractionType::ScrubbingLeftHandle:
-		TrimFadeTransform->StartTime = SelectedTime;
-		NotifyTransformationPropertyChanged(TrimFadeTransform, GET_MEMBER_NAME_CHECKED(UWaveformTransformationTrimFade, StartTime), DesiredChangeType);
+		StartTimeHandle->SetValue(SelectedTime, Flags);
 		break;
 	case FWaveformTransformationTrimFadeRenderer::ETrimFadeInteractionType::ScrubbingRightHandle:
-		TrimFadeTransform->EndTime = SelectedTime;
-		NotifyTransformationPropertyChanged(TrimFadeTransform, GET_MEMBER_NAME_CHECKED(UWaveformTransformationTrimFade, EndTime), DesiredChangeType);
+		EndTimeHandle->SetValue(SelectedTime, Flags);
 		break;
 	case FWaveformTransformationTrimFadeRenderer::ETrimFadeInteractionType::ScrubbingFadeIn:
-		TrimFadeTransform->StartFadeTime = FMath::Clamp(SelectedTime - TrimFadeTransform->StartTime, 0.f, TNumericLimits<float>().Max());
-		NotifyTransformationPropertyChanged(TrimFadeTransform, GET_MEMBER_NAME_CHECKED(UWaveformTransformationTrimFade, StartFadeTime), DesiredChangeType);
+		{
+			float StartFadeTimeValue = FMath::Clamp(SelectedTime - GetPropertyValue<double>(StartTimeHandle), 0.f, TNumericLimits<float>().Max());
+			StartFadeTimeHandle->SetValue(StartFadeTimeValue, Flags);
+		}
 		break;
 	case FWaveformTransformationTrimFadeRenderer::ETrimFadeInteractionType::ScrubbingFadeOut:
-		TrimFadeTransform->EndFadeTime = FMath::Clamp(TrimFadeTransform->EndTime - SelectedTime, 0.f, TNumericLimits<float>().Max());
-		NotifyTransformationPropertyChanged(TrimFadeTransform, GET_MEMBER_NAME_CHECKED(UWaveformTransformationTrimFade, EndFadeTime), DesiredChangeType);
+		{	
+			float EndFadeTimeValue = FMath::Clamp(GetPropertyValue<double>(EndTimeHandle) - SelectedTime, 0.f, TNumericLimits<float>().Max());
+			EndFadeTimeHandle->SetValue(EndFadeTimeValue, Flags);
+		}
 		break;
 	default:
 		break;
 	}
 }
 
-
 FWaveformTransformationTrimFadeRenderer::ETrimFadeInteractionType FWaveformTransformationTrimFadeRenderer::GetInteractionTypeFromCursorPosition(const FVector2D& InLocalCursorPosition, const FGeometry& WidgetGeometry) const
 {
 	if (IsCursorInFadeInInteractionRange(InLocalCursorPosition, WidgetGeometry))
@@ -333,4 +348,32 @@ bool FWaveformTransformationTrimFadeRenderer::IsCursorInFadeOutInteractionRange(
 {
 	return FadeOutInteractionXRange.Contains(InLocalCursorPosition.X)
 		&& InLocalCursorPosition.Y < WidgetGeometry.GetLocalSize().Y* InteractionRatioYDelta;
+}
+
+FText FWaveformTransformationTrimFadeRenderer::GetPropertyEditedByCurrentInteraction() const
+{
+	FText OutPropertyName = FText::FromName(NAME_None);
+
+	switch (TrimFadeInteractionType)
+	{
+	case FWaveformTransformationTrimFadeRenderer::ETrimFadeInteractionType::None:
+		break;
+	case FWaveformTransformationTrimFadeRenderer::ETrimFadeInteractionType::ScrubbingLeftHandle:
+		OutPropertyName = FText::FromName(GET_MEMBER_NAME_CHECKED(UWaveformTransformationTrimFade, StartTime));
+		break;
+	case FWaveformTransformationTrimFadeRenderer::ETrimFadeInteractionType::ScrubbingRightHandle:
+		OutPropertyName = FText::FromName(GET_MEMBER_NAME_CHECKED(UWaveformTransformationTrimFade, EndTime));
+		break;
+	case FWaveformTransformationTrimFadeRenderer::ETrimFadeInteractionType::ScrubbingFadeIn:
+		OutPropertyName = FText::FromName(GET_MEMBER_NAME_CHECKED(UWaveformTransformationTrimFade, StartFadeTime));
+		break;
+	case FWaveformTransformationTrimFadeRenderer::ETrimFadeInteractionType::ScrubbingFadeOut:
+		OutPropertyName = FText::FromName(GET_MEMBER_NAME_CHECKED(UWaveformTransformationTrimFade, EndFadeTime));
+		break;
+	default:
+		break;
+	}
+
+	return MoveTemp(OutPropertyName);
+
 }
\ No newline at end of file
diff --git a/Engine/Plugins/Editor/WaveformEditor/Source/WaveformTransformationsWidgets/Private/WaveformTransformationTrimFadeRenderer.h b/Engine/Plugins/Editor/WaveformEditor/Source/WaveformTransformationsWidgets/Private/WaveformTransformationTrimFadeRenderer.h
--- a/Engine/Plugins/Editor/WaveformEditor/Source/WaveformTransformationsWidgets/Private/WaveformTransformationTrimFadeRenderer.h
+++ b/Engine/Plugins/Editor/WaveformEditor/Source/WaveformTransformationsWidgets/Private/WaveformTransformationTrimFadeRenderer.h
@@ -3,12 +3,11 @@
 #pragma once
 
 #include "WaveformTransformationRendererBase.h"
-#include "WaveformTransformationTrimFade.h"
 
 class FWaveformTransformationTrimFadeRenderer : public FWaveformTransformationRendererBase
 {
 public:
-	explicit FWaveformTransformationTrimFadeRenderer(const TObjectPtr<UWaveformTransformationTrimFade> TransformationToRender);
+	FWaveformTransformationTrimFadeRenderer() = default;
 
 	virtual int32 OnPaint(const FPaintArgs& Args, const FGeometry& AllottedGeometry, const FSlateRect& MyCullingRect, FSlateWindowElementList& OutDrawElements, int32 LayerId, const FWidgetStyle& InWidgetStyle, bool bParentEnabled) const override;
 	virtual void Tick(const FGeometry& AllottedGeometry, const double InCurrentTime, const float InDeltaTime) override;
@@ -18,6 +17,8 @@ public:
 	virtual FReply OnMouseWheel(SWidget& OwnerWidget, const FGeometry& MyGeometry, const FPointerEvent& MouseEvent) override;
 	virtual FCursorReply OnCursorQuery(const FGeometry& MyGeometry, const FPointerEvent& CursorEvent) const override;
 
+	virtual void SetPropertyHandles(const TArray<TSharedRef<IPropertyHandle>>& InPropertyHandles) override;
+
 private:
 	int32 DrawTrimHandles(const FGeometry& AllottedGeometry, FSlateWindowElementList& OutDrawElements, int32 LayerId) const;
 	int32 DrawFadeCurves(const FGeometry& AllottedGeometry, FSlateWindowElementList& OutDrawElements, int32 LayerId) const;
@@ -27,10 +28,9 @@ private:
 	bool IsCursorInFadeInInteractionRange(const FVector2D& InLocalCursorPosition, const FGeometry& WidgetGeometry) const;
 	bool IsCursorInFadeOutInteractionRange(const FVector2D& InLocalCursorPosition, const FGeometry& WidgetGeometry) const;
 
+	FText GetPropertyEditedByCurrentInteraction() const;
 	FVector2D GetLocalCursorPosition(const FPointerEvent& MouseEvent, const FGeometry& EventGeometry) const;
-	float ConvertXRatioToTime(const float InRatio) const;
-
-	TObjectPtr<UWaveformTransformationTrimFade> TrimFadeTransform = nullptr;
+	double ConvertXRatioToTime(const float InRatio) const;
 
 	float StartTimeHandleX = 0.f;
 	float EndTimeHandleX = 0.f;
@@ -62,5 +62,12 @@ private:
 	} TrimFadeInteractionType = ETrimFadeInteractionType::None;
 
 	ETrimFadeInteractionType GetInteractionTypeFromCursorPosition(const FVector2D& InLocalCursorPosition, const FGeometry& WidgetGeometry) const;
-	void SetPropertyValueDependingOnInteractionType(const FPointerEvent& MouseEvent, const FGeometry& WidgetGeometry, const EPropertyChangeType::Type DesiredChangeType);
-};
\ No newline at end of file
+	void SetPropertyValueDependingOnInteractionType(const FPointerEvent& MouseEvent, const FGeometry& WidgetGeometry, const EPropertyValueSetFlags::Type Flags = EPropertyValueSetFlags::DefaultFlags);
+
+	TSharedPtr<IPropertyHandle> StartTimeHandle;
+	TSharedPtr<IPropertyHandle> EndTimeHandle;
+	TSharedPtr<IPropertyHandle> StartFadeTimeHandle;
+	TSharedPtr<IPropertyHandle> StartFadeCurveHandle;
+	TSharedPtr<IPropertyHandle> EndFadeTimeHandle;
+	TSharedPtr<IPropertyHandle> EndFadeCurveHandle;
+};
diff --git a/Engine/Plugins/Editor/WaveformEditor/Source/WaveformTransformationsWidgets/Private/WaveformTransformationsRenderManager.cpp b/Engine/Plugins/Editor/WaveformEditor/Source/WaveformTransformationsWidgets/Private/WaveformTransformationsRenderManager.cpp
--- a/Engine/Plugins/Editor/WaveformEditor/Source/WaveformTransformationsWidgets/Private/WaveformTransformationsRenderManager.cpp
+++ b/Engine/Plugins/Editor/WaveformEditor/Source/WaveformTransformationsWidgets/Private/WaveformTransformationsRenderManager.cpp
@@ -7,12 +7,14 @@
 #include "DSP/FloatArrayMath.h"
 #include "IWaveformTransformationRenderer.h"
 #include "Sound/SoundWave.h"
+#include "WaveformTransformationDurationRenderer.h"
 
-FWaveformTransformationsRenderManager::FWaveformTransformationsRenderManager(TObjectPtr<USoundWave> InSoundWave, 
-	TSharedRef<FWaveformEditorRenderData> InWaveformRenderData, 
-	TFunction<void(FPropertyChangedEvent&, FEditPropertyChain*)> InPropertyChangeNotifier)
+FWaveformTransformationsRenderManager::FWaveformTransformationsRenderManager(
+	TObjectPtr<USoundWave> InSoundWave, 
+	FOnTransformationsPropertiesRequired InTransformationChainPropertiesHandler)
 	: SoundWaveToRender(InSoundWave)
-	, LayersFactory(MakeUnique<FWaveformTransformationRenderLayerFactory>(InWaveformRenderData, InPropertyChangeNotifier))
+	, LayersFactory(MakeUnique<FWaveformTransformationRenderLayerFactory>())
+	, ChainPropertiesHandler(InTransformationChainPropertiesHandler)
 {
 	GenerateLayersChain();	
 }
@@ -26,21 +28,29 @@ void FWaveformTransformationsRenderManager::GenerateLayersChain()
 
 	if (SoundWaveToRender->Transformations.Num() > 0)
 	{
+		FTransformationsToPropertiesArray TransformPropertiesMap;
+
 		for (const TObjectPtr<UWaveformTransformationBase> Transformation : SoundWaveToRender->Transformations)
 		{
 			if (Transformation != nullptr)
 			{
-				TransformationsToRender.Add(Transformation);
-				
-				TSharedPtr<IWaveformTransformationRenderer> TransformationUI = LayersFactory->Create(Transformation);
-				FTransformationRenderLayerInfo RenderLayerInfo = FTransformationRenderLayerInfo(TransformationUI, FTransformationLayerConstraints(0.f, 1.f));
-				RenderLayers.Add(RenderLayerInfo);
-				
+				TArray<TSharedRef<IPropertyHandle>> TransformationPropertiesArray;
+				TransformPropertiesMap.Emplace(Transformation, TransformationPropertiesArray);
 			}
 		}
 
-		CreateDurationHighlightLayer();
+		ChainPropertiesHandler.Execute(TransformPropertiesMap);
+		
 
+		for (FTransformationToPropertiesPair& ObjPropsPair : TransformPropertiesMap)
+		{
+			TransformationsToRender.Add(ObjPropsPair.Key);
+			TSharedPtr<IWaveformTransformationRenderer> TransformationUI = LayersFactory->Create(ObjPropsPair.Key, ObjPropsPair.Value);
+			FTransformationRenderLayerInfo RenderLayerInfo = FTransformationRenderLayerInfo(TransformationUI, FTransformationLayerConstraints(0.f, 1.f));
+			RenderLayers.Add(RenderLayerInfo);
+		}
+
+		CreateDurationHighlightLayer();
 	}
 
 	OnLayersChainGenerated.Broadcast(RenderLayers.GetData(), RenderLayers.Num());
@@ -66,6 +76,7 @@ void FWaveformTransformationsRenderManager::GenerateRenderDataInternal()
 	}
 
 	uint32 NumOriginalSamples = RawPCMData.Num() * sizeof(uint8) / sizeof(int16);
+	NumOriginalWaveformFrames = NumOriginalSamples / NumChannels;
 	uint32 FirstEditedSample = 0;
 	uint32 LastEditedSample = NumOriginalSamples;
 
@@ -122,6 +133,8 @@ void FWaveformTransformationsRenderManager::GenerateRenderDataInternal()
 		check(DurationHiglightLayer)
 		FWaveformTransformationRenderInfo DurationLayerInfo{ TransformationInfo.SampleRate, TransformationInfo.NumChannels, FirstEditedSample, LastEditedSample - FirstEditedSample };
 		DurationHiglightLayer->SetTransformationWaveInfo(MoveTemp(DurationLayerInfo));
+		DurationHiglightLayer->SetOriginalWaveformFrames(NumOriginalWaveformFrames);
+		
 			
 		if (!bChainChangesFileLength)
 		{
@@ -182,9 +195,9 @@ void FWaveformTransformationsRenderManager::CreateDurationHighlightLayer()
 {
 	if (!DurationHiglightLayer)
 	{
-		DurationHiglightLayer = LayersFactory->CreateDurationHiglightLayer();
+		DurationHiglightLayer = MakeShared<FWaveformTransformationDurationRenderer>(NumOriginalWaveformFrames);
 	}
 
 	DurationHiglightLayer->SetTransformationWaveInfo(FWaveformTransformationRenderInfo());
 	RenderLayers.Emplace(DurationHiglightLayer, FTransformationLayerConstraints(0.f,1.f));
-}
\ No newline at end of file
+}
diff --git a/Engine/Plugins/Editor/WaveformEditor/Source/WaveformTransformationsWidgets/Private/WaveformTransformationsWidgetsModule.cpp b/Engine/Plugins/Editor/WaveformEditor/Source/WaveformTransformationsWidgets/Private/WaveformTransformationsWidgetsModule.cpp
--- a/Engine/Plugins/Editor/WaveformEditor/Source/WaveformTransformationsWidgets/Private/WaveformTransformationsWidgetsModule.cpp
+++ b/Engine/Plugins/Editor/WaveformEditor/Source/WaveformTransformationsWidgets/Private/WaveformTransformationsWidgetsModule.cpp
@@ -1,6 +1,20 @@
 // Copyright Epic Games, Inc. All Rights Reserved.
 
 #include "WaveformTransformationsWidgetsModule.h"
+
 #include "Modules/ModuleManager.h"
+#include "WaveformTransformationRendererMapper.h"
+#include "WaveformTransformationTrimFade.h"
+#include "WaveformTransformationTrimFadeRenderer.h"
+
+void FWaveformTransformationsWidgetsModule::StartupModule()
+{
+	FWaveformTransformationRendererMapper& RendererMapper = FWaveformTransformationRendererMapper::Get();
+	RendererMapper.RegisterRenderer<FWaveformTransformationTrimFadeRenderer>(UWaveformTransformationTrimFade::StaticClass());
+}
+
+void FWaveformTransformationsWidgetsModule::ShutdownModule()
+{
+}
 
-IMPLEMENT_MODULE(FWaveformTransformationsWidgetsModule, WaveformTransformationsWidgets);
\ No newline at end of file
+IMPLEMENT_MODULE(FWaveformTransformationsWidgetsModule, WaveformTransformationsWidgets)
\ No newline at end of file
diff --git a/Engine/Plugins/Editor/WaveformEditor/Source/WaveformTransformationsWidgets/Public/WaveformTransformationRenderLayerFactory.h b/Engine/Plugins/Editor/WaveformEditor/Source/WaveformTransformationsWidgets/Public/WaveformTransformationRenderLayerFactory.h
deleted file mode 100644
diff --git a/Engine/Plugins/Editor/WaveformEditor/Source/WaveformTransformationsWidgets/Public/WaveformTransformationsRenderManager.h b/Engine/Plugins/Editor/WaveformEditor/Source/WaveformTransformationsWidgets/Public/WaveformTransformationsRenderManager.h
--- a/Engine/Plugins/Editor/WaveformEditor/Source/WaveformTransformationsWidgets/Public/WaveformTransformationsRenderManager.h
+++ b/Engine/Plugins/Editor/WaveformEditor/Source/WaveformTransformationsWidgets/Public/WaveformTransformationsRenderManager.h
@@ -7,13 +7,17 @@
 #include "WaveformTransformationRenderLayerFactory.h"
 
 class IWaveformTransformationRenderer;
+class FWaveformTransformationDurationRenderer;
 class USoundWave;
 class UWaveformTransformationBase;
 
 using FTransformationLayerConstraints = TPair<float, float>;
 using FTransformationRenderLayerInfo = TPair<TSharedPtr<IWaveformTransformationRenderer>, FTransformationLayerConstraints>;
+using FTransformationToPropertiesPair = TPair<TObjectPtr<UWaveformTransformationBase>, TArray<TSharedRef<IPropertyHandle>>>;
+using FTransformationsToPropertiesArray = TArray<FTransformationToPropertiesPair>;
 
 DECLARE_MULTICAST_DELEGATE(FOnRenderElementsUpdated)
+DECLARE_DELEGATE_OneParam(FOnTransformationsPropertiesRequired, FTransformationsToPropertiesArray& /*Transformations Properties To Hand*/)
 DECLARE_MULTICAST_DELEGATE_TwoParams(FOnLayersChainGenerated, FTransformationRenderLayerInfo* /*First Layer ptr*/, const int32 /* NLayers */)
 DECLARE_MULTICAST_DELEGATE_SixParams(FOnRenderDataGenerated, const uint8* /* First RawPCMData element */, const uint32 /*Num Samples*/, const uint32 /*First Edited Sample*/, const uint32 /*Last Edited Sample*/, const uint32/*Sample Rate*/, const uint16 /*Num Channels*/)
 
@@ -45,9 +49,8 @@ class WAVEFORMTRANSFORMATIONSWIDGETS_API FWaveformTransformationsRenderManager
 {
 public:
 	explicit FWaveformTransformationsRenderManager(
-		TObjectPtr<USoundWave> InSoundWave, 
-		TSharedRef<FWaveformEditorRenderData> InWaveformRenderData, 
-		TFunction<void(FPropertyChangedEvent&, FEditPropertyChain*)> InPropertyChangeNotifier
+		TObjectPtr<USoundWave> InSoundWave,
+		FOnTransformationsPropertiesRequired InTransformationChainPropertiesHandler
 	);
 
 	/** Used to generate the stack of transformations UI					*/
@@ -85,5 +88,9 @@ private:
 	TArray<uint8> RawPCMData;
 
 	TUniquePtr<FWaveformTransformationRenderLayerFactory> LayersFactory = nullptr;
-	TSharedPtr<IWaveformTransformationRenderer> DurationHiglightLayer = nullptr;
+	TSharedPtr<FWaveformTransformationDurationRenderer> DurationHiglightLayer = nullptr;
+
+	FOnTransformationsPropertiesRequired ChainPropertiesHandler;
+
+	uint32 NumOriginalWaveformFrames = 0;
 };
\ No newline at end of file
diff --git a/Engine/Plugins/Editor/WaveformEditor/Source/WaveformTransformationsWidgets/Public/WaveformTransformationsWidgetsModule.h b/Engine/Plugins/Editor/WaveformEditor/Source/WaveformTransformationsWidgets/Public/WaveformTransformationsWidgetsModule.h
--- a/Engine/Plugins/Editor/WaveformEditor/Source/WaveformTransformationsWidgets/Public/WaveformTransformationsWidgetsModule.h
+++ b/Engine/Plugins/Editor/WaveformEditor/Source/WaveformTransformationsWidgets/Public/WaveformTransformationsWidgetsModule.h
@@ -6,4 +6,7 @@
 
 class FWaveformTransformationsWidgetsModule : public IModuleInterface
 {
+	/** IModuleInterface implementation */
+	virtual void StartupModule() override;
+	virtual void ShutdownModule() override;
 };
\ No newline at end of file
diff --git a/Engine/Plugins/Editor/WaveformEditor/Source/WaveformTransformationsWidgets/WaveformTransformationsWidgets.build.cs b/Engine/Plugins/Editor/WaveformEditor/Source/WaveformTransformationsWidgets/WaveformTransformationsWidgets.build.cs
--- a/Engine/Plugins/Editor/WaveformEditor/Source/WaveformTransformationsWidgets/WaveformTransformationsWidgets.build.cs
+++ b/Engine/Plugins/Editor/WaveformEditor/Source/WaveformTransformationsWidgets/WaveformTransformationsWidgets.build.cs
@@ -20,10 +20,12 @@ namespace UnrealBuildTool.Rules
 					"EditorStyle",
 					"Engine",
 					"InputCore",
+					"PropertyEditor",
 					"Slate",
 					"SignalProcessing",
 					"SlateCore", 
 					"ToolMenus", 
+					"UnrealEd",
 					"WaveformTransformations",
 					"WaveformEditorWidgets",
 				}
