Fix LocText redundant keys in wave editor instantiator error messages

diff --git a/Engine/Plugins/Editor/WaveformEditor/Source/WaveformEditor/Private/WaveformEditorInstantiator.cpp b/Engine/Plugins/Editor/WaveformEditor/Source/WaveformEditor/Private/WaveformEditorInstantiator.cpp
--- a/Engine/Plugins/Editor/WaveformEditor/Source/WaveformEditor/Private/WaveformEditorInstantiator.cpp
+++ b/Engine/Plugins/Editor/WaveformEditor/Source/WaveformEditor/Private/WaveformEditorInstantiator.cpp
@@ -70,7 +70,7 @@ bool FWaveformEditorInstantiator::CanSoundWaveBeOpenedInEditor(const USoundWave*
 
 	if (SoundWaveToEdit == nullptr)
 	{
-		ErrorText = LOCTEXT("WaveformEditorOpeningError", "Could not open waveform editor. Selected SoundWave was null.");
+		ErrorText = LOCTEXT("WaveformEditorOpeningError_NullSoundWave", "Could not open waveform editor. Selected SoundWave was null.");
 		bCanOpenWaveEditor = false;
 	}
 	else
@@ -79,19 +79,19 @@ bool FWaveformEditorInstantiator::CanSoundWaveBeOpenedInEditor(const USoundWave*
 		
 		if (SoundWaveToEdit->GetDuration() == 0.f)
 		{
-			ErrorText = FText::Format(LOCTEXT("WaveformEditorOpeningError", "Could not open waveform editor for SoundWave '{0}': duration is 0"), SoundWaveNameText);
+			ErrorText = FText::Format(LOCTEXT("WaveformEditorOpeningError_ZeroDuration", "Could not open waveform editor for SoundWave '{0}': duration is 0"), SoundWaveNameText);
 			bCanOpenWaveEditor = false;
 		}
 
 		if (SoundWaveToEdit->NumChannels == 0)
 		{
-			ErrorText = FText::Format(LOCTEXT("WaveformEditorOpeningError", "Could not open waveform editor for SoundWave '{0}': channel count is 0"), SoundWaveNameText);
+			ErrorText = FText::Format(LOCTEXT("WaveformEditorOpeningError_ZeroChannels", "Could not open waveform editor for SoundWave '{0}': channel count is 0"), SoundWaveNameText);
 			bCanOpenWaveEditor = false;
 		}
 
 		if (SoundWaveToEdit->TotalSamples == 0)
 		{
-			ErrorText = FText::Format(LOCTEXT("WaveformEditorOpeningError", "Could not open waveform editor for SoundWave '{0}': found 0 total samples.\n\nConsider reimporting the asset to fix it."), SoundWaveNameText);
+			ErrorText = FText::Format(LOCTEXT("WaveformEditorOpeningError_ZeroSamples", "Could not open waveform editor for SoundWave '{0}': found 0 total samples.\n\nConsider reimporting the asset to fix it."), SoundWaveNameText);
 			bCanOpenWaveEditor = false;
 		}
 	}
