[Waveform Editor] quick fix for playhead going outside transformed waveform bounce when stopping #trivial #rb none

diff --git a/Engine/Plugins/Editor/WaveformEditor/Source/WaveformEditorWidgets/Private/WaveformEditorTransportCoordinator.cpp b/Engine/Plugins/Editor/WaveformEditor/Source/WaveformEditorWidgets/Private/WaveformEditorTransportCoordinator.cpp
--- a/Engine/Plugins/Editor/WaveformEditor/Source/WaveformEditorWidgets/Private/WaveformEditorTransportCoordinator.cpp
+++ b/Engine/Plugins/Editor/WaveformEditor/Source/WaveformEditorWidgets/Private/WaveformEditorTransportCoordinator.cpp
@@ -269,7 +269,6 @@ void FWaveformEditorTransportCoordinator::UpdateZoomRatioAndDisplayRange(const f
 void FWaveformEditorTransportCoordinator::Stop()
 {
 	ReceivePlayBackRatio(0.f);
-	UpdateDisplayRange(0.f, ZoomRatio);
 }
 
 void FWaveformEditorTransportCoordinator::UpdateDisplayRange(const float MinValue, const float MaxValue)
@@ -301,4 +300,4 @@ uint32 FWaveformEditorTransportCoordinator::GetSampleFromPlayheadPosition(const
 const float FWaveformEditorTransportCoordinator::GetPlayheadPosition() const
 {
 	return PlayheadPosition;
-}
+}
\ No newline at end of file
