Generic/Dumb Sampled Sequence Viewer Widget derived from FWaveformViewer
- widget only gets the data it needs to draw
- removed any dependency from FWaveformData
- created a "GridService" interface, to provid grid data and snapping to multiple widgets from a unified tool

diff --git a/Engine/Plugins/Editor/WaveformEditor/Source/WaveformEditorWidgets/Private/SWaveformEditorTimeRuler.cpp b/Engine/Plugins/Editor/WaveformEditor/Source/WaveformEditorWidgets/Private/SWaveformEditorTimeRuler.cpp
--- a/Engine/Plugins/Editor/WaveformEditor/Source/WaveformEditorWidgets/Private/SWaveformEditorTimeRuler.cpp
+++ b/Engine/Plugins/Editor/WaveformEditor/Source/WaveformEditorWidgets/Private/SWaveformEditorTimeRuler.cpp
@@ -5,16 +5,15 @@
 #include "Layout/WidgetPath.h"
 #include "Framework/Application/SlateApplication.h"
 #include "Framework/MultiBox/MultiBoxBuilder.h"
-#include "WaveformEditorRenderData.h"
 #include "WaveformEditorTransportCoordinator.h"
 
-
 #define LOCTEXT_NAMESPACE "WaveformEditorTimeRuler"
 
-void SWaveformEditorTimeRuler::Construct(const FArguments& InArgs, TSharedRef<FWaveformEditorTransportCoordinator> InTransportCoordinator, TSharedRef<FWaveformEditorRenderData> InRenderData)
+void SWaveformEditorTimeRuler::Construct(const FArguments& InArgs, TSharedRef<FWaveformEditorTransportCoordinator> InTransportCoordinator, TSharedRef<ISampledSequenceGridService> InGridService)
 {
 	TransportCoordinator = InTransportCoordinator;
-	RenderData = InRenderData;
+	GridService = InGridService;
+	UpdateGridMetrics();
 
 	DisplayUnit = InArgs._DisplayUnit;
 
@@ -99,12 +98,12 @@ void SWaveformEditorTimeRuler::DrawRulerTicks(const FGeometry& AllottedGeometry,
 
 		}
 
-		float TickTimeSeconds = MajorTickX / GridMetrics.PixelsPerSecond + GridMetrics.StartTime;
-		DrawTickTimeString(TickTimeSeconds, MajorTickX, MajorTickY, OutDrawElements, LayerId, AllottedGeometry);
+		uint32 TickFrame = FMath::RoundToInt32(MajorTickX / GridMetrics.PixelsPerFrame) + GridMetrics.StartFrame;
+		DrawTickTimeString(TickFrame, MajorTickX, MajorTickY, OutDrawElements, LayerId, AllottedGeometry);
 	}
 }
 
-void SWaveformEditorTimeRuler::DrawTickTimeString(float TickTimeSeconds, const double TickX, const double TickY, FSlateWindowElementList& OutDrawElements, int32& LayerId, const FGeometry& AllottedGeometry) const
+void SWaveformEditorTimeRuler::DrawTickTimeString(uint32 TickFrame, const double TickX, const double TickY, FSlateWindowElementList& OutDrawElements, int32& LayerId, const FGeometry& AllottedGeometry) const
 {
 	FString TimeString;
 
@@ -112,15 +111,15 @@ void SWaveformEditorTimeRuler::DrawTickTimeString(float TickTimeSeconds, const d
 	{
 	case EWaveformEditorDisplayUnit::Samples:
 	{
-		uint32 SampleTime = TickTimeSeconds * RenderData->GetSampleRate();
+		uint32 SampleTime = TickFrame;
 		TimeString = FString::Printf(TEXT("%d"), SampleTime);
 	}
 	break;
 	case EWaveformEditorDisplayUnit::Seconds:
-		TimeString = FString::Printf(TEXT("%.3f"), TickTimeSeconds);
+		TimeString = FString::Printf(TEXT("%.3f"), TickFrame / (float) GridMetrics.SampleRate);
 		break;
 	default:
-		TimeString = FString::Printf(TEXT("%.3f"), TickTimeSeconds);
+		TimeString = FString::Printf(TEXT("%.3f"), TickFrame / (float)GridMetrics.SampleRate);
 		break;
 	}
 
@@ -202,9 +201,10 @@ void SWaveformEditorTimeRuler::NotifyTimeUnitMenuSelection(const EWaveformEditor
 	}
 }
 
-void SWaveformEditorTimeRuler::UpdateGridMetrics(const FWaveEditorGridMetrics& InMetrics)
+void SWaveformEditorTimeRuler::UpdateGridMetrics()
 {
-	GridMetrics = InMetrics;
+	check(GridService)
+	GridMetrics = GridService->GetGridMetrics();
 }
 
 void SWaveformEditorTimeRuler::UpdateDisplayUnit(const EWaveformEditorDisplayUnit InDisplayUnit)
@@ -238,7 +238,8 @@ void SWaveformEditorTimeRuler::OnStyleUpdated(const FWaveformEditorWidgetStyleBa
 void SWaveformEditorTimeRuler::DrawPlayheadHandle(const FGeometry& AllottedGeometry, FSlateWindowElementList& OutDrawElements, int32& LayerId) const
 {
 	const float WindowWidth = AllottedGeometry.Size.X;
-	const float	HandleStart = TransportCoordinator->GetPlayheadPosition() * WindowWidth - HandleWidth / 2;
+	const float SnappedHandleCenter = GridService->SnapPositionToClosestFrame(TransportCoordinator->GetPlayheadPosition() * WindowWidth);
+	const float	HandleStart = SnappedHandleCenter - HandleWidth / 2;
 	const float	HandleEnd = HandleStart + HandleWidth;
  	FPaintGeometry HandleGeometry = AllottedGeometry.ToPaintGeometry(FVector2f(HandleEnd - HandleStart, AllottedGeometry.Size.Y), FSlateLayoutTransform(FVector2f(HandleStart, 0)));
 
diff --git a/Engine/Plugins/Editor/WaveformEditor/Source/WaveformEditorWidgets/Private/SWaveformEditorTimeRuler.h b/Engine/Plugins/Editor/WaveformEditor/Source/WaveformEditorWidgets/Private/SWaveformEditorTimeRuler.h
--- a/Engine/Plugins/Editor/WaveformEditor/Source/WaveformEditorWidgets/Private/SWaveformEditorTimeRuler.h
+++ b/Engine/Plugins/Editor/WaveformEditor/Source/WaveformEditorWidgets/Private/SWaveformEditorTimeRuler.h
@@ -2,17 +2,16 @@
 
 #pragma once
 
+#include "ISampledSequenceGridService.h"
 #include "Framework/MultiBox/MultiBoxBuilder.h"
 #include "Styling/SlateStyleRegistry.h"
 #include "Styling/SlateWidgetStyleAsset.h"
 #include "WaveformEditorDisplayUnit.h"
-#include "WaveformEditorGridData.h"
 #include "WaveformEditorSlateTypes.h"
 #include "Widgets/SCompoundWidget.h"
 
 DECLARE_MULTICAST_DELEGATE_OneParam(FOnTimeUnitMenuSelection, const EWaveformEditorDisplayUnit /* Requested Display Unit */);
 
-class FWaveformEditorRenderData;
 class FWaveformEditorTransportCoordinator;
 
 class SWaveformEditorTimeRuler : public SCompoundWidget
@@ -30,8 +29,8 @@ class SWaveformEditorTimeRuler : public SCompoundWidget
 	SLATE_END_ARGS()
 
 public:
-	void Construct(const FArguments& InArgs, TSharedRef<FWaveformEditorTransportCoordinator> InTransportCoordinator, TSharedRef<FWaveformEditorRenderData> InRenderData);
-	void UpdateGridMetrics(const FWaveEditorGridMetrics& InMetrics);
+	void Construct(const FArguments& InArgs, TSharedRef<FWaveformEditorTransportCoordinator> InTransportCoordinator, TSharedRef<ISampledSequenceGridService> InGridService);
+	void UpdateGridMetrics();
 	void UpdateDisplayUnit(const EWaveformEditorDisplayUnit InDisplayUnit);
 
 	void OnStyleUpdated(const FWaveformEditorWidgetStyleBase* UpdatedStyle);
@@ -49,13 +48,13 @@ private:
 
 	void DrawPlayheadHandle(const FGeometry& AllottedGeometry, FSlateWindowElementList& OutDrawElements, int32& LayerId) const;
 	void DrawRulerTicks(const FGeometry& AllottedGeometry, FSlateWindowElementList& OutDrawElements, int32& LayerId) const;
-	void DrawTickTimeString(float TickTimeSeconds, const double TickX, const double TickY, FSlateWindowElementList& OutDrawElements, int32& LayerId, const FGeometry& AllottedGeometry) const;
+	void DrawTickTimeString(uint32 TickFrame, const double TickX, const double TickY, FSlateWindowElementList& OutDrawElements, int32& LayerId, const FGeometry& AllottedGeometry) const;
 
 	TSharedRef<SWidget> MakeContextMenu();
 	void MakeTimeUnitsSubMenu(FMenuBuilder& SubMenuBuilder);
 	void NotifyTimeUnitMenuSelection(const EWaveformEditorDisplayUnit SelectedDisplayUnit) const;
 
-	FWaveEditorGridMetrics GridMetrics;
+	FSampledSequenceGridMetrics GridMetrics;
 
 	const FWaveformEditorTimeRulerStyle* Style = nullptr;
 
@@ -72,7 +71,7 @@ private:
 	FSlateFontInfo TicksTextFont;
 
 	TSharedPtr<FWaveformEditorTransportCoordinator> TransportCoordinator = nullptr;
-	TSharedPtr<FWaveformEditorRenderData> RenderData = nullptr;
+	TSharedPtr<ISampledSequenceGridService> GridService = nullptr;
 
 	EWaveformEditorDisplayUnit DisplayUnit;
 };
diff --git a/Engine/Plugins/Editor/WaveformEditor/Source/WaveformEditorWidgets/Private/SWaveformPanel.cpp b/Engine/Plugins/Editor/WaveformEditor/Source/WaveformEditorWidgets/Private/SWaveformPanel.cpp
--- a/Engine/Plugins/Editor/WaveformEditor/Source/WaveformEditorWidgets/Private/SWaveformPanel.cpp
+++ b/Engine/Plugins/Editor/WaveformEditor/Source/WaveformEditorWidgets/Private/SWaveformPanel.cpp
@@ -4,7 +4,7 @@
 
 #include "SWaveformEditorTimeRuler.h"
 #include "SWaveformTransformationsOverlay.h"
-#include "SWaveformViewer.h"
+#include "SSampledSequenceViewer.h"
 #include "SWaveformViewerOverlay.h"
 #include "WaveformEditorDisplayUnit.h"
 #include "WaveformEditorGridData.h"
@@ -23,8 +23,16 @@ void SWaveformPanel::Construct(const FArguments& InArgs, TSharedRef<FWaveformEdi
 	WaveformEditorStyle = &FWaveformEditorStyle::Get();
 	check(WaveformEditorStyle);
 
-	SetUpGridData(InRenderData, InTransportCoordinator);
-	SetUpWaveformViewer(InRenderData, InTransportCoordinator, InZoomManager, GridData.ToSharedRef());
+	RenderData = InRenderData;
+	GenerateFloatRenderData();
+	InRenderData->OnRenderDataUpdated.AddSP(this, &SWaveformPanel::OnRenderDataUpdated);
+
+	TransportCoordinator = InTransportCoordinator;
+	TransportCoordinator->OnDisplayRangeUpdated.AddSP(this, &SWaveformPanel::OnDisplayRangeUpdated);
+
+	SetUpGridData(InRenderData);
+	SetUpWaveformViewer(GridData.ToSharedRef(), InRenderData);
+	SetUpZoomManager(InZoomManager, InTransportCoordinator);
 
 	if (InWaveformTransformationsOverlay)
 	{
@@ -32,7 +40,7 @@ void SWaveformPanel::Construct(const FArguments& InArgs, TSharedRef<FWaveformEdi
 	}
 	
 	SetUpWaveformViewerOverlay(InTransportCoordinator, InZoomManager);
-	SetUpTimeRuler(InRenderData, InTransportCoordinator, GridData.ToSharedRef());
+	SetUpTimeRuler(InTransportCoordinator, GridData.ToSharedRef());
 	CreateLayout();
 }
 
@@ -75,14 +83,13 @@ void SWaveformPanel::CreateLayout()
 	];
 }
 
-void SWaveformPanel::SetUpTimeRuler(TSharedRef<FWaveformEditorRenderData> InRenderData, TSharedRef<FWaveformEditorTransportCoordinator> InTransportCoordinator, TSharedRef<FWaveformEditorGridData> InGridData)
+void SWaveformPanel::SetUpTimeRuler(TSharedRef<FWaveformEditorTransportCoordinator> InTransportCoordinator, TSharedRef<FWaveformEditorGridData> InGridData)
 {
 	FWaveformEditorTimeRulerStyle* TimeRulerStyle = &WaveformEditorStyle->GetRegisteredWidgetStyle<FWaveformEditorTimeRulerStyle>("WaveformEditorRuler.Style").Get();
 	check(TimeRulerStyle);
 
-	TimeRuler = SNew(SWaveformEditorTimeRuler, InTransportCoordinator, InRenderData).DisplayUnit(DisplayUnit).Style(TimeRulerStyle);
+	TimeRuler = SNew(SWaveformEditorTimeRuler, InTransportCoordinator, InGridData).DisplayUnit(DisplayUnit).Style(TimeRulerStyle);
 	TimeRulerStyle->OnStyleUpdated.AddSP(TimeRuler.ToSharedRef(), &SWaveformEditorTimeRuler::OnStyleUpdated);
-	InGridData->OnGridMetricsUpdated.AddSP(TimeRuler.ToSharedRef(), &SWaveformEditorTimeRuler::UpdateGridMetrics);
 	TimeRuler->OnTimeUnitMenuSelection.AddSP(this, &SWaveformPanel::UpdateDisplayUnit);
 }
 
@@ -91,26 +98,27 @@ void SWaveformPanel::SetUpWaveformViewerOverlay(TSharedRef<FWaveformEditorTransp
 	FWaveformViewerOverlayStyle* OverlayStyle = &WaveformEditorStyle->GetRegisteredWidgetStyle<FWaveformViewerOverlayStyle>("WaveformViewerOverlay.Style").Get();
 	check(OverlayStyle);
 
-	WaveformViewerOverlay = SNew(SWaveformViewerOverlay, InTransportCoordinator, WaveformTransformationsOverlay.ToSharedRef()).Style(OverlayStyle);
+	WaveformViewerOverlay = SNew(SWaveformViewerOverlay, InTransportCoordinator, WaveformTransformationsOverlay.ToSharedRef(), GridData.ToSharedRef()).Style(OverlayStyle);
 	OverlayStyle->OnStyleUpdated.AddSP(WaveformViewerOverlay.ToSharedRef(), &SWaveformViewerOverlay::OnStyleUpdated);
 	WaveformViewerOverlay->OnNewMouseDelta.BindSP(InZoomManager, &FWaveformEditorZoomController::ZoomByDelta);
 }
 
-void SWaveformPanel::SetUpWaveformViewer(TSharedRef<FWaveformEditorRenderData> InRenderData, TSharedRef<FWaveformEditorTransportCoordinator> InTransportCoordinator, TSharedRef<FWaveformEditorZoomController> InZoomManager, TSharedRef<FWaveformEditorGridData> InGridData)
+void SWaveformPanel::SetUpWaveformViewer(TSharedRef<FWaveformEditorGridData> InGridData, TSharedRef<FWaveformEditorRenderData> InRenderData)
 {
-	FWaveformViewerStyle* WaveViewerStyle = &WaveformEditorStyle->GetRegisteredWidgetStyle<FWaveformViewerStyle>("WaveformViewer.Style").Get();
+	FSampledSequenceViewerStyle* WaveViewerStyle = &WaveformEditorStyle->GetRegisteredWidgetStyle<FSampledSequenceViewerStyle>("WaveformViewer.Style").Get();
 	check(WaveViewerStyle);
 	
-	WaveformViewer = SNew(SWaveformViewer, InRenderData, InTransportCoordinator).Style(WaveViewerStyle);
-	WaveViewerStyle->OnStyleUpdated.AddSP(WaveformViewer.ToSharedRef(), &SWaveformViewer::OnStyleUpdated);
-	InZoomManager->OnZoomRatioChanged.AddSP(InTransportCoordinator, &FWaveformEditorTransportCoordinator::OnZoomLevelChanged);
-	InGridData->OnGridMetricsUpdated.AddSP(WaveformViewer.ToSharedRef(), &SWaveformViewer::UpdateGridMetrics);
+	InGridData->UpdateGridMetrics(WaveViewerStyle->DesiredWidth);
+	TimeSeriesDrawingUtils::FSampledSequenceDrawingParams WaveformViewerDrawingParams;
+	WaveformViewerDrawingParams.MaxDisplayedValue = TNumericLimits<int16>::Max();
+
+	WaveformViewer = SNew(SSampledSequenceViewer, MakeArrayView(FloatRenderData.GetData(), FloatRenderData.Num()), InRenderData->GetNumChannels(), InGridData).Style(WaveViewerStyle).SequenceDrawingParams(WaveformViewerDrawingParams);
+	WaveViewerStyle->OnStyleUpdated.AddSP(WaveformViewer.ToSharedRef(), &SSampledSequenceViewer::OnStyleUpdated);
 }
 
-void SWaveformPanel::SetUpGridData(TSharedRef<FWaveformEditorRenderData> InRenderData, TSharedRef<FWaveformEditorTransportCoordinator> InTransportCoordinator)
+void SWaveformPanel::SetUpGridData(TSharedRef<FWaveformEditorRenderData> InRenderData)
 {
-	GridData = MakeShared<FWaveformEditorGridData>(InRenderData);
-	InTransportCoordinator->OnDisplayRangeUpdated.AddSP(GridData.ToSharedRef(), &FWaveformEditorGridData::UpdateDisplayRange);
+	GridData = MakeShared<FWaveformEditorGridData>(InRenderData->GetNumSamples() / InRenderData->GetNumChannels(), InRenderData->GetSampleRate());
 
 	const ISlateStyle* WaveEditorStyle = FSlateStyleRegistry::FindSlateStyle("WaveformEditorStyle");
 	if (ensure(WaveEditorStyle))
@@ -120,6 +128,62 @@ void SWaveformPanel::SetUpGridData(TSharedRef<FWaveformEditorRenderData> InRende
 	}
 }
 
+void SWaveformPanel::SetUpZoomManager(TSharedRef<FWaveformEditorZoomController> InZoomManager, TSharedRef<FWaveformEditorTransportCoordinator> InTransportCoordinator)
+{
+	InZoomManager->OnZoomRatioChanged.AddSP(InTransportCoordinator, &FWaveformEditorTransportCoordinator::OnZoomLevelChanged);
+}
+
+void SWaveformPanel::OnRenderDataUpdated()
+{
+	check(TransportCoordinator)
+	OnDisplayRangeUpdated(TransportCoordinator->GetDisplayRange());
+}
+
+void SWaveformPanel::OnDisplayRangeUpdated(const TRange<float> NewDisplayRange)
+{
+	check (RenderData)
+
+	const uint8 MinFramesToDisplay = 1;
+	const uint32 MinSamplesToDisplay = MinFramesToDisplay * RenderData->GetNumChannels();
+	const uint32 NumOriginalSamples = RenderData->GetSampleData().Num();
+	const uint32 NumOriginalFrames = NumOriginalSamples / RenderData->GetNumChannels();
+
+	const uint32 FirstRenderedSample = FMath::Clamp(FMath::RoundToInt32(NumOriginalFrames * NewDisplayRange.GetLowerBoundValue()), 0, NumOriginalFrames - MinFramesToDisplay) * RenderData->GetNumChannels();	
+	const uint32 NumFramesToRender = FMath::RoundToInt32(NumOriginalFrames * NewDisplayRange.Size<float>());
+	const uint32 NumSamplesToRender = FMath::Clamp(NumFramesToRender * RenderData->GetNumChannels(), MinSamplesToDisplay, NumOriginalSamples - FirstRenderedSample);
+	
+	check(NumSamplesToRender % RenderData->GetNumChannels() == 0 && FirstRenderedSample % RenderData->GetNumChannels() == 0);
+	
+	if (GridData)
+	{
+		const uint32 FirstRenderedFrame = FirstRenderedSample / RenderData->GetNumChannels();
+		GridData->UpdateDisplayRange(TRange<uint32>(FirstRenderedFrame, FirstRenderedFrame + NumFramesToRender));
+	}
+
+	if (TimeRuler)
+	{
+		TimeRuler->UpdateGridMetrics();
+	}
+
+	if (WaveformViewer)
+	{
+		TArrayView<const float> RenderedView = MakeArrayView(FloatRenderData.GetData(), FloatRenderData.Num());
+		RenderedView = RenderedView.Slice(FirstRenderedSample, NumSamplesToRender);
+		WaveformViewer->UpdateView(RenderedView, RenderData->GetNumChannels());
+	}
+}
+
+void SWaveformPanel::GenerateFloatRenderData()
+{
+	TArrayView<const int16> SampleData = RenderData->GetSampleData();
+	FloatRenderData.SetNumUninitialized(SampleData.Num());
+
+	for (int32 Sample = 0; Sample < SampleData.Num(); ++Sample)
+	{
+		FloatRenderData[Sample] = SampleData[Sample];
+	}
+}
+
 void SWaveformPanel::UpdateDisplayUnit(const EWaveformEditorDisplayUnit InDisplayUnit)
 {
 	DisplayUnit = InDisplayUnit;
@@ -132,7 +196,11 @@ void SWaveformPanel::Tick(const FGeometry& AllottedGeometry, const double InCurr
 
 	if (PaintedWidth != CachedPixelWidth)
 	{
-		GridData->UpdateGridMetrics(PaintedWidth);
 		CachedPixelWidth = PaintedWidth;
+
+		if (GridData)
+		{
+			GridData->UpdateGridMetrics(PaintedWidth);
+		}	
 	}
 }
\ No newline at end of file
diff --git a/Engine/Plugins/Editor/WaveformEditor/Source/WaveformEditorWidgets/Private/SWaveformViewer.cpp b/Engine/Plugins/Editor/WaveformEditor/Source/WaveformEditorWidgets/Private/SWaveformViewer.cpp
deleted file mode 100644
diff --git a/Engine/Plugins/Editor/WaveformEditor/Source/WaveformEditorWidgets/Private/SWaveformViewer.h b/Engine/Plugins/Editor/WaveformEditor/Source/WaveformEditorWidgets/Private/SWaveformViewer.h
deleted file mode 100644
diff --git a/Engine/Plugins/Editor/WaveformEditor/Source/WaveformEditorWidgets/Private/SWaveformViewerOverlay.cpp b/Engine/Plugins/Editor/WaveformEditor/Source/WaveformEditorWidgets/Private/SWaveformViewerOverlay.cpp
--- a/Engine/Plugins/Editor/WaveformEditor/Source/WaveformEditorWidgets/Private/SWaveformViewerOverlay.cpp
+++ b/Engine/Plugins/Editor/WaveformEditor/Source/WaveformEditorWidgets/Private/SWaveformViewerOverlay.cpp
@@ -2,13 +2,16 @@
 
 #include "SWaveformViewerOverlay.h"
 
+#include "ISampledSequenceGridService.h"
 #include "WaveformEditorTransportCoordinator.h"
+#include "WaveformEditorGridData.h"
 #include "Widgets/SLeafWidget.h"
 
-void SWaveformViewerOverlay::Construct(const FArguments& InArgs, TSharedRef<FWaveformEditorTransportCoordinator> InTransportCoordinator, TSharedRef<SWaveformTransformationsOverlay> InTransformationsOverlay)
+void SWaveformViewerOverlay::Construct(const FArguments& InArgs, TSharedRef<FWaveformEditorTransportCoordinator> InTransportCoordinator, TSharedRef<SWaveformTransformationsOverlay> InTransformationsOverlay, TSharedRef<ISampledSequenceGridService> InGridService)
 {
 	TransportCoordinator = InTransportCoordinator;
 	TransformationsOverlay = InTransformationsOverlay;
+	GridService = InGridService;
 
 	Style = InArgs._Style;
 	
@@ -80,7 +83,7 @@ int32 SWaveformViewerOverlay::OnPaint(const FPaintArgs& Args, const FGeometry& A
 
 int32 SWaveformViewerOverlay::DrawPlayhead(const FGeometry& AllottedGeometry, FSlateWindowElementList& OutDrawElements, int32 LayerId) const
 {
-	const float PlayHeadX = AllottedGeometry.GetLocalSize().X * TransportCoordinator->GetPlayheadPosition();
+	const float PlayHeadX = GridService->SnapPositionToClosestFrame(AllottedGeometry.GetLocalSize().X * TransportCoordinator->GetPlayheadPosition());
 
 	TArray<FVector2D> LinePoints;
 	{
diff --git a/Engine/Plugins/Editor/WaveformEditor/Source/WaveformEditorWidgets/Private/SWaveformViewerOverlay.h b/Engine/Plugins/Editor/WaveformEditor/Source/WaveformEditorWidgets/Private/SWaveformViewerOverlay.h
--- a/Engine/Plugins/Editor/WaveformEditor/Source/WaveformEditorWidgets/Private/SWaveformViewerOverlay.h
+++ b/Engine/Plugins/Editor/WaveformEditor/Source/WaveformEditorWidgets/Private/SWaveformViewerOverlay.h
@@ -11,6 +11,9 @@
 
 DECLARE_DELEGATE_OneParam(FOnNewMouseDelta, const float /* new delta */)
 
+class ISampledSequenceGridService;
+class FWaveformEditorGridData;
+
 class SWaveformViewerOverlay : public SLeafWidget
 {
 public:
@@ -22,7 +25,7 @@ public:
 
 	SLATE_END_ARGS()
 
-	void Construct(const FArguments& InArgs, TSharedRef<FWaveformEditorTransportCoordinator> InTransportCoordinator, TSharedRef<SWaveformTransformationsOverlay> InTransformationsOverlay);
+	void Construct(const FArguments& InArgs, TSharedRef<FWaveformEditorTransportCoordinator> InTransportCoordinator, TSharedRef<SWaveformTransformationsOverlay> InTransformationsOverlay, TSharedRef<ISampledSequenceGridService> InGridService);
 
 	virtual FReply OnMouseButtonDown(const FGeometry& MyGeometry, const FPointerEvent& MouseEvent) override;
 	virtual FReply OnMouseWheel(const FGeometry& MyGeometry, const FPointerEvent& MouseEvent) override;
@@ -41,6 +44,7 @@ private:
 
 	TSharedPtr<FWaveformEditorTransportCoordinator> TransportCoordinator = nullptr;
 	TSharedPtr<SWaveformTransformationsOverlay> TransformationsOverlay = nullptr;
+	TSharedPtr<ISampledSequenceGridService> GridService = nullptr;
 
 	const FWaveformViewerOverlayStyle* Style = nullptr;
 
diff --git a/Engine/Plugins/Editor/WaveformEditor/Source/WaveformEditorWidgets/Private/WaveformEditorGridData.cpp b/Engine/Plugins/Editor/WaveformEditor/Source/WaveformEditorWidgets/Private/WaveformEditorGridData.cpp
--- a/Engine/Plugins/Editor/WaveformEditor/Source/WaveformEditorWidgets/Private/WaveformEditorGridData.cpp
+++ b/Engine/Plugins/Editor/WaveformEditor/Source/WaveformEditorWidgets/Private/WaveformEditorGridData.cpp
@@ -5,17 +5,18 @@
 #include "Fonts/FontMeasure.h"
 #include "Framework/Application/SlateApplication.h"
 #include "GenericPlatform/GenericPlatformMath.h"
-#include "Misc/FrameRate.h"
-#include "WaveformEditorRenderData.h"
 
-FWaveformEditorGridData::FWaveformEditorGridData(TSharedRef<FWaveformEditorRenderData> InRenderData, const FSlateFontInfo* InTicksTimeFont /* = nullptr*/)
-	: RenderData(InRenderData)
-	, TicksTimeFont(InTicksTimeFont)
+FWaveformEditorGridData::FWaveformEditorGridData(const uint32 InTotalFrames, const uint32 InSampleRateHz, const FSlateFontInfo* InTicksTimeFont /*= nullptr*/)
+ 	: TotalFrames(InTotalFrames)
+	, DisplayRange(TRange<uint32>::Inclusive(0, InTotalFrames))
+ 	, TicksTimeFont(InTicksTimeFont)
+ 	, GridFrameRate(InSampleRateHz, 1)
 {
 }
 
-void FWaveformEditorGridData::UpdateDisplayRange(const TRange<float> InDisplayRange)
+void FWaveformEditorGridData::UpdateDisplayRange(const TRange<uint32> InDisplayRange)
 {
+	check(DisplayRange.Size<uint32>() >= 2)
 	DisplayRange = InDisplayRange;
 	UpdateGridMetrics(GridPixelWidth);
 }
@@ -26,10 +27,10 @@ bool FWaveformEditorGridData::UpdateGridMetrics(const float InGridPixelWidth)
 	{
 		GridPixelWidth = InGridPixelWidth;
 	}
-	
-	const float WaveformDurationSeconds = RenderData->GetOriginalWaveformDurationInSeconds();
-	const double StartTime = WaveformDurationSeconds * DisplayRange.GetLowerBoundValue();
-	const double DisplayedDuration = WaveformDurationSeconds * DisplayRange.Size<double>();
+
+	const float WaveformDurationSeconds = TotalFrames / (double) GridFrameRate.Numerator;
+	const double StartTimeSeconds = DisplayRange.GetLowerBoundValue() / (double) GridFrameRate.Numerator;	
+	const double DisplayedDuration = (DisplayRange.Size<double>() - 1) / GridFrameRate.Numerator; //we account for one less frame so we can distribute them evenly on the screen
 	const double PixelsPerSecond = GridPixelWidth / DisplayedDuration;
 	double MajorGridStepSeconds = 0.0;
 	float MinTicksPixelSpacing = 30.0f;
@@ -42,28 +43,24 @@ bool FWaveformEditorGridData::UpdateGridMetrics(const float InGridPixelWidth)
 		MinTicksPixelSpacing = MaxTextSize.X;
 	}
 
-	const FFrameRate FrameRate(RenderData->GetSampleRate(), 1);
 
-	if (!FrameRate.ComputeGridSpacing(PixelsPerSecond, MajorGridStepSeconds, GridMetrics.NumMinorGridDivisions, MinTicksPixelSpacing + 5.f))
+	if (!GridFrameRate.ComputeGridSpacing(PixelsPerSecond, MajorGridStepSeconds, GridMetrics.NumMinorGridDivisions, MinTicksPixelSpacing + 5.f))
 	{
 		return false;
 	}
 
-	const double GridOffset = FGenericPlatformMath::Fmod(StartTime, MajorGridStepSeconds);
+	const bool bZeroGridOffset = DisplayRange.Size<uint32>() < GridPixelWidth;
+	const double GridOffset = bZeroGridOffset ? 0 : FGenericPlatformMath::Fmod(StartTimeSeconds, MajorGridStepSeconds);
 	GridMetrics.FirstMajorTickX = (0.f - GridOffset) * PixelsPerSecond;
 	GridMetrics.MajorGridXStep = MajorGridStepSeconds * PixelsPerSecond;
-	GridMetrics.StartTime = StartTime;
-	GridMetrics.PixelsPerSecond = PixelsPerSecond;
-
-	if (OnGridMetricsUpdated.IsBound())
-	{
-		OnGridMetricsUpdated.Broadcast(GridMetrics);
-	}
+	GridMetrics.StartFrame = DisplayRange.GetLowerBoundValue();
+	GridMetrics.PixelsPerFrame =  PixelsPerSecond / GridFrameRate.Numerator;
+	GridMetrics.SampleRate = GridFrameRate.Numerator;
 
 	return true;
 }
 
-const FWaveEditorGridMetrics FWaveformEditorGridData::GetGridMetrics() const
+const FSampledSequenceGridMetrics FWaveformEditorGridData::GetGridMetrics() const
 {
 	return GridMetrics;
 }
@@ -72,3 +69,22 @@ void FWaveformEditorGridData::SetTicksTimeFont(const FSlateFontInfo* InNewFont)
 {
 	TicksTimeFont = InNewFont;
 }
+
+const float FWaveformEditorGridData::SnapPositionToClosestFrame(const float InPixelPosition) const
+{
+	const float DistanceFromPreviousSample = FGenericPlatformMath::Fmod(InPixelPosition, GridMetrics.PixelsPerFrame);
+	const bool bSnapToNext = DistanceFromPreviousSample > GridMetrics.PixelsPerFrame / 2.f;
+	float SnappedPosition = 0.f;
+
+	if (bSnapToNext)
+	{
+		SnappedPosition = InPixelPosition + (GridMetrics.PixelsPerFrame - DistanceFromPreviousSample);
+	}
+	else
+	{
+		SnappedPosition = InPixelPosition -  DistanceFromPreviousSample;
+	}
+
+
+	return SnappedPosition;
+}
\ No newline at end of file
diff --git a/Engine/Plugins/Editor/WaveformEditor/Source/WaveformEditorWidgets/Private/WaveformEditorSlateTypes.cpp b/Engine/Plugins/Editor/WaveformEditor/Source/WaveformEditorWidgets/Private/WaveformEditorSlateTypes.cpp
--- a/Engine/Plugins/Editor/WaveformEditor/Source/WaveformEditorWidgets/Private/WaveformEditorSlateTypes.cpp
+++ b/Engine/Plugins/Editor/WaveformEditor/Source/WaveformEditorWidgets/Private/WaveformEditorSlateTypes.cpp
@@ -11,39 +11,10 @@ namespace WaveformEditorStylesSharedParams
 	const FLazyName HandleBrushName = TEXT("Sequencer.Timeline.VanillaScrubHandleDown");
 	const FLinearColor PlayheadColor = FLinearColor(255.f, 0.1f, 0.2f, 1.f);
 	const FLinearColor RulerTicksColor = FLinearColor(1.f, 1.f, 1.f, 0.9f);
-	const FLinearColor WaveformBackGroundColor = FLinearColor(0.02f, 0.02f, 0.02f, 1.f);
 	const float ViewerHeight = 720.f;
 	const float ViewerWidth = 1280.f;
 }
 
-const FName FWaveformViewerStyle::TypeName("FWaveformViewerStyle");
-
-FWaveformViewerStyle::FWaveformViewerStyle()
-	: WaveformColor(FLinearColor::White)
-	, WaveformLineThickness(1.f)
-	, MajorGridLineColor(FLinearColor::Black)
-	, MinorGridLineColor(FLinearColor(0.f, 0.f, 0.f, 0.5f))
-	, ZeroCrossingLineColor(FLinearColor::Black)
-	, ZeroCrossingLineThickness(1.f)
-	, SampleMarkersSize(2.5f)
-	, WaveformBackgroundColor(WaveformEditorStylesSharedParams::WaveformBackGroundColor)
-	, BackgroundBrush(*FAppStyle::GetBrush(WaveformEditorStylesSharedParams::BackgroundBrushName))
-	, DesiredWidth(WaveformEditorStylesSharedParams::ViewerWidth)
-	, DesiredHeight(WaveformEditorStylesSharedParams::ViewerHeight)
-{
-}
-
-const FWaveformViewerStyle& FWaveformViewerStyle::GetDefault()
-{
-	static FWaveformViewerStyle Default;
-	return Default;
-}
-
-void FWaveformViewerStyle::GetResources(TArray< const FSlateBrush* >& OutBrushes) const
-{
-	OutBrushes.Add(&BackgroundBrush);
-}
-
 const FName FWaveformViewerOverlayStyle::TypeName("FWaveformViewerOverlayStyle");
 
 FWaveformViewerOverlayStyle::FWaveformViewerOverlayStyle()
diff --git a/Engine/Plugins/Editor/WaveformEditor/Source/WaveformEditorWidgets/Private/WaveformEditorStyle.cpp b/Engine/Plugins/Editor/WaveformEditor/Source/WaveformEditorWidgets/Private/WaveformEditorStyle.cpp
--- a/Engine/Plugins/Editor/WaveformEditor/Source/WaveformEditorWidgets/Private/WaveformEditorStyle.cpp
+++ b/Engine/Plugins/Editor/WaveformEditor/Source/WaveformEditorWidgets/Private/WaveformEditorStyle.cpp
@@ -2,6 +2,7 @@
 
 #include "WaveformEditorStyle.h"
 
+#include "AudioWidgetsSlateTypes.h"
 #include "Styling/SlateStyleRegistry.h"
 #include "WaveformEditorSlateTypes.h"
 #include "WaveformEditorWidgetsSettings.h"
@@ -45,10 +46,10 @@ void FWaveformEditorStyle::Init()
 	StyleInstance->SetParentStyleName("CoreStyle");
 
 	//Waveform Viewer style
-	FWaveformViewerStyle WaveViewerStyle = FWaveformViewerStyle()
-		.SetWaveformColor(Settings->WaveformColor)
+	FSampledSequenceViewerStyle WaveViewerStyle = FSampledSequenceViewerStyle()
+		.SetSequenceColor(Settings->WaveformColor)
 		.SetBackgroundColor(Settings->WaveformBackgroundColor)
-		.SetWaveformLineThickness(Settings->WaveformLineThickness)
+		.SetSequenceLineThickness(Settings->WaveformLineThickness)
 		.SetSampleMarkersSize(Settings->SampleMarkersSize)
 		.SetMajorGridLineColor(Settings->MajorGridColor)
 		.SetMinorGridLineColor(Settings->MinorGridColor)
@@ -85,13 +86,13 @@ const UWaveformEditorWidgetsSettings* FWaveformEditorStyle::GetWidgetsSettings()
 
 void FWaveformEditorStyle::OnWidgetSettingsUpdated(const FName& PropertyName, const UWaveformEditorWidgetsSettings* Settings)
 {
-	TSharedRef<FWaveformViewerStyle> WaveformViewerStyle = GetRegisteredWidgetStyle<FWaveformViewerStyle>(WaveformViewerStyleName);
+	TSharedRef<FSampledSequenceViewerStyle> WaveformViewerStyle = GetRegisteredWidgetStyle<FSampledSequenceViewerStyle>(WaveformViewerStyleName);
 	TSharedRef<FWaveformViewerOverlayStyle> WaveformViewerOverlayStyle = GetRegisteredWidgetStyle<FWaveformViewerOverlayStyle>(WaveformViewerOverlayStyleName);
 	TSharedRef<FWaveformEditorTimeRulerStyle> WaveformEditorTimeRulerStyle = GetRegisteredWidgetStyle<FWaveformEditorTimeRulerStyle>(WaveformEditorRulerStyleName);
 
 	if (PropertyName == GET_MEMBER_NAME_CHECKED(UWaveformEditorWidgetsSettings, WaveformColor))
 	{
-		WaveformViewerStyle->SetWaveformColor(Settings->WaveformColor);
+		WaveformViewerStyle->SetSequenceColor(Settings->WaveformColor);
 	}
 	else if (PropertyName == GET_MEMBER_NAME_CHECKED(UWaveformEditorWidgetsSettings, WaveformBackgroundColor))
 	{
@@ -99,7 +100,7 @@ void FWaveformEditorStyle::OnWidgetSettingsUpdated(const FName& PropertyName, co
 	}
 	else if (PropertyName == GET_MEMBER_NAME_CHECKED(UWaveformEditorWidgetsSettings, WaveformLineThickness))
 	{
-		WaveformViewerStyle->SetWaveformLineThickness(Settings->WaveformLineThickness);
+		WaveformViewerStyle->SetSequenceLineThickness(Settings->WaveformLineThickness);
 	}
 	else if (PropertyName == GET_MEMBER_NAME_CHECKED(UWaveformEditorWidgetsSettings, SampleMarkersSize))
 	{
diff --git a/Engine/Plugins/Editor/WaveformEditor/Source/WaveformEditorWidgets/Public/SWaveformPanel.h b/Engine/Plugins/Editor/WaveformEditor/Source/WaveformEditorWidgets/Public/SWaveformPanel.h
--- a/Engine/Plugins/Editor/WaveformEditor/Source/WaveformEditorWidgets/Public/SWaveformPanel.h
+++ b/Engine/Plugins/Editor/WaveformEditor/Source/WaveformEditorWidgets/Public/SWaveformPanel.h
@@ -12,7 +12,7 @@ class FWaveformEditorTransportCoordinator;
 class FWaveformEditorZoomController;
 class SWaveformEditorTimeRuler;
 class SWaveformTransformationsOverlay;
-class SWaveformViewer;
+class SSampledSequenceViewer;
 class SWaveformViewerOverlay;
 enum class EWaveformEditorDisplayUnit;
 
@@ -30,16 +30,27 @@ public:
 private: 
 	void CreateLayout();
 
-	void SetUpWaveformViewer(TSharedRef<FWaveformEditorRenderData> InRenderData, TSharedRef<FWaveformEditorTransportCoordinator> InTransportCoordinator, TSharedRef<FWaveformEditorZoomController> InZoomManager, TSharedRef<FWaveformEditorGridData> InGridData);
+	void SetUpWaveformViewer(TSharedRef<FWaveformEditorGridData> InGridData, TSharedRef<FWaveformEditorRenderData> InRenderData);
 	void SetUpWaveformViewerOverlay(TSharedRef<FWaveformEditorTransportCoordinator> InTransportCoordinator, TSharedRef<FWaveformEditorZoomController> InZoomManager);
-	void SetUpTimeRuler(TSharedRef<FWaveformEditorRenderData> InRenderData, TSharedRef<FWaveformEditorTransportCoordinator> InTransportCoordinator, TSharedRef<FWaveformEditorGridData> InGridData);
-	void SetUpGridData(TSharedRef<FWaveformEditorRenderData> InRenderData, TSharedRef<FWaveformEditorTransportCoordinator> InTransportCoordinator);
+	void SetUpTimeRuler(TSharedRef<FWaveformEditorTransportCoordinator> InTransportCoordinator, TSharedRef<FWaveformEditorGridData> InGridData);
+	void SetUpGridData(TSharedRef<FWaveformEditorRenderData> InRenderData);
+	void SetUpZoomManager(TSharedRef<FWaveformEditorZoomController> InZoomManager, TSharedRef<FWaveformEditorTransportCoordinator> InTransportCoordinator);
+
+	void OnRenderDataUpdated();
+	void OnDisplayRangeUpdated(const TRange<float> NewDisplayRange);
+
+	//temporary float render data conversion while widget are made generic
+	void GenerateFloatRenderData();
+	TArray<float> FloatRenderData;
 
 	void UpdateDisplayUnit(const EWaveformEditorDisplayUnit InDisplayUnit);
 
+	TSharedPtr<FWaveformEditorRenderData> RenderData = nullptr;
 	TSharedPtr<FWaveformEditorGridData> GridData = nullptr;
+	TSharedPtr<FWaveformEditorTransportCoordinator> TransportCoordinator = nullptr;
+
 	TSharedPtr<SWaveformEditorTimeRuler> TimeRuler = nullptr;
-	TSharedPtr<SWaveformViewer> WaveformViewer = nullptr;
+	TSharedPtr<SSampledSequenceViewer> WaveformViewer = nullptr;
 	TSharedPtr<SWaveformTransformationsOverlay> WaveformTransformationsOverlay = nullptr;
 	TSharedPtr<SWaveformViewerOverlay> WaveformViewerOverlay = nullptr;
 
@@ -48,4 +59,5 @@ private:
 	EWaveformEditorDisplayUnit DisplayUnit;
 
 	FWaveformEditorStyle* WaveformEditorStyle;
+
 };
\ No newline at end of file
diff --git a/Engine/Plugins/Editor/WaveformEditor/Source/WaveformEditorWidgets/Public/WaveformEditorGridData.h b/Engine/Plugins/Editor/WaveformEditor/Source/WaveformEditorWidgets/Public/WaveformEditorGridData.h
--- a/Engine/Plugins/Editor/WaveformEditor/Source/WaveformEditorWidgets/Public/WaveformEditorGridData.h
+++ b/Engine/Plugins/Editor/WaveformEditor/Source/WaveformEditorWidgets/Public/WaveformEditorGridData.h
@@ -4,39 +4,29 @@
 
 #include "Delegates/Delegate.h"
 #include "HAL/Platform.h"
+#include "ISampledSequenceGridService.h"
 #include "Math/Range.h"
+#include "Misc/FrameRate.h"
 
 struct FSlateFontInfo;
-class FWaveformEditorRenderData;
 
-struct WAVEFORMEDITORWIDGETS_API FWaveEditorGridMetrics
-{
-	int32 NumMinorGridDivisions = 0;
-	double PixelsPerSecond = 0;
-	double FirstMajorTickX = 0;
-	double MajorGridXStep = 0;
-	double StartTime = 0;
-};
-
-DECLARE_MULTICAST_DELEGATE_OneParam(FOnGridMetricsUpdated, const FWaveEditorGridMetrics& /* New Grid Metrics */);
-
-class WAVEFORMEDITORWIDGETS_API FWaveformEditorGridData
+class WAVEFORMEDITORWIDGETS_API FWaveformEditorGridData : public ISampledSequenceGridService
 {
 public: 
-	explicit FWaveformEditorGridData(TSharedRef<FWaveformEditorRenderData> InRenderData, const FSlateFontInfo* InTicksTimeFont = nullptr);
-
-	FOnGridMetricsUpdated OnGridMetricsUpdated;
+	explicit FWaveformEditorGridData(const uint32 InTotalFrames, const uint32 InSampleRateHz, const FSlateFontInfo* InTicksTimeFont = nullptr);
 
-	void UpdateDisplayRange(const TRange<float> InDisplayRange);
+	void UpdateDisplayRange(const TRange<uint32> InDisplayRange);
 	bool UpdateGridMetrics(const float InGridPixelWidth);
-	const FWaveEditorGridMetrics GetGridMetrics() const;
+	virtual const FSampledSequenceGridMetrics GetGridMetrics() const override;
 	void SetTicksTimeFont(const FSlateFontInfo* InNewFont);
+	const float SnapPositionToClosestFrame(const float InPixelPosition) const;
 
 private:
-	FWaveEditorGridMetrics GridMetrics;
-	TSharedPtr<FWaveformEditorRenderData> RenderData = nullptr;
-	TRange<float> DisplayRange = TRange<float>::Inclusive(0.f, 1.f);
+	FSampledSequenceGridMetrics GridMetrics;
+	uint32 TotalFrames = 0;
+	TRange<uint32> DisplayRange;
 
 	float GridPixelWidth = 0.f;
 	const FSlateFontInfo* TicksTimeFont = nullptr;
+	FFrameRate GridFrameRate;
 };
\ No newline at end of file
diff --git a/Engine/Plugins/Editor/WaveformEditor/Source/WaveformEditorWidgets/Public/WaveformEditorSlateTypes.h b/Engine/Plugins/Editor/WaveformEditor/Source/WaveformEditorWidgets/Public/WaveformEditorSlateTypes.h
--- a/Engine/Plugins/Editor/WaveformEditor/Source/WaveformEditorWidgets/Public/WaveformEditorSlateTypes.h
+++ b/Engine/Plugins/Editor/WaveformEditor/Source/WaveformEditorWidgets/Public/WaveformEditorSlateTypes.h
@@ -23,67 +23,6 @@ struct FWaveformEditorWidgetStyleBase : public FSlateWidgetStyle
 	FOnWidgetStyleUpdated OnStyleUpdated;
 };
 
-/**
- * Represents the appearance of a Waveform Viewer
- */
-USTRUCT(BlueprintType)
-struct FWaveformViewerStyle : public FWaveformEditorWidgetStyleBase
-{
-	GENERATED_USTRUCT_BODY()
-
-	FWaveformViewerStyle();
-
-	static const FWaveformViewerStyle& GetDefault();
-	virtual const FName GetTypeName() const override { return TypeName; };
-	virtual void GetResources(TArray< const FSlateBrush* >& OutBrushes) const override;
-
-	static const FName TypeName;
-
-	UPROPERTY(EditAnywhere, BlueprintReadWrite, Category = Appearance)
-	FSlateColor WaveformColor;
-	FWaveformViewerStyle& SetWaveformColor(const FSlateColor InWaveformColor) { WaveformColor = InWaveformColor; BroadcastStyleUpdate(); return *this; }
-
-	UPROPERTY(EditAnywhere, BlueprintReadWrite, Category = Appearance)
-	float WaveformLineThickness;
-	FWaveformViewerStyle& SetWaveformLineThickness(const float InWaveformLineThickness) { WaveformLineThickness = InWaveformLineThickness; BroadcastStyleUpdate(); return *this; }
-
-	UPROPERTY(EditAnywhere, BlueprintReadWrite, Category = Appearance)
-	FSlateColor MajorGridLineColor;
-	FWaveformViewerStyle& SetMajorGridLineColor(const FSlateColor InMajorGridLineColor) { MajorGridLineColor = InMajorGridLineColor; BroadcastStyleUpdate(); return *this; }
-
-	UPROPERTY(EditAnywhere, BlueprintReadWrite, Category = Appearance)
-	FSlateColor MinorGridLineColor;
-	FWaveformViewerStyle& SetMinorGridLineColor(const FSlateColor InMinorGridLineColor) { MinorGridLineColor = InMinorGridLineColor; BroadcastStyleUpdate(); return *this; }
-
-	UPROPERTY(EditAnywhere, BlueprintReadWrite, Category = Appearance)
-	FSlateColor ZeroCrossingLineColor;
-	FWaveformViewerStyle& SetZeroCrossingLineColor(const FSlateColor InZeroCrossingLineColor) { ZeroCrossingLineColor = InZeroCrossingLineColor; BroadcastStyleUpdate(); return *this; }
-
-	UPROPERTY(EditAnywhere, BlueprintReadWrite, Category = Appearance)
-	float ZeroCrossingLineThickness;
-	FWaveformViewerStyle& SetZeroCrossingLineThickness(const float InZeroCrossingLineThickness) { ZeroCrossingLineThickness = InZeroCrossingLineThickness; BroadcastStyleUpdate(); return *this; }
-
-	UPROPERTY(EditAnywhere, BlueprintReadWrite, Category = Appearance)
-	float SampleMarkersSize;
-	FWaveformViewerStyle& SetSampleMarkersSize(const float InSampleMarkersSize) { SampleMarkersSize = InSampleMarkersSize; BroadcastStyleUpdate(); return *this; }
-
-	UPROPERTY(EditAnywhere, BlueprintReadWrite, Category = Appearance)
-	FSlateColor WaveformBackgroundColor;
-	FWaveformViewerStyle& SetBackgroundColor(const FSlateColor InBackgroundColor) { WaveformBackgroundColor = InBackgroundColor; BroadcastStyleUpdate(); return *this; }
-
-	UPROPERTY(EditAnywhere, BlueprintReadWrite, Category = Appearance)
-	FSlateBrush BackgroundBrush;
-	FWaveformViewerStyle& SetBackgroundBrush(const FSlateBrush InBackgroundBrush) { BackgroundBrush = InBackgroundBrush; BroadcastStyleUpdate(); return *this; }
-	
-	UPROPERTY(EditAnywhere, BlueprintReadWrite, Category = Appearance)
-	float DesiredWidth;
-	FWaveformViewerStyle& SetDesiredWidth(const float InDesiredWidth) { DesiredWidth = InDesiredWidth; BroadcastStyleUpdate(); return *this; }
-
-	UPROPERTY(EditAnywhere, BlueprintReadWrite, Category = Appearance)
-	float DesiredHeight;
-	FWaveformViewerStyle& SetDesiredHeight(const float InDesiredHeight) { DesiredHeight = InDesiredHeight; BroadcastStyleUpdate(); return *this; }
-};
-
 /**
  * Represents the appearance of a Waveform Viewer Overlay style
  */
diff --git a/Engine/Plugins/Runtime/AudioWidgets/Source/AudioWidgets/Private/AudioWidgetsSlateTypes.cpp b/Engine/Plugins/Runtime/AudioWidgets/Source/AudioWidgets/Private/AudioWidgetsSlateTypes.cpp
--- a/Engine/Plugins/Runtime/AudioWidgets/Source/AudioWidgets/Private/AudioWidgetsSlateTypes.cpp
+++ b/Engine/Plugins/Runtime/AudioWidgets/Source/AudioWidgets/Private/AudioWidgetsSlateTypes.cpp
@@ -3,6 +3,7 @@
 #include "AudioWidgetsSlateTypes.h"
 #include "Brushes/SlateNoResource.h"
 #include "Brushes/SlateRoundedBoxBrush.h"
+#include "Styling/AppStyle.h"
 #include "Styling/StyleColors.h"
 
 #include UE_INLINE_GENERATED_CPP_BY_NAME(AudioWidgetsSlateTypes)
@@ -72,3 +73,31 @@ const FAudioRadialSliderStyle& FAudioRadialSliderStyle::GetDefault()
 	return Default;
 }
 
+const FName FSampledSequenceViewerStyle::TypeName("FSampledSequenceViewerStyle");
+
+FSampledSequenceViewerStyle::FSampledSequenceViewerStyle()
+	: SequenceColor(FLinearColor::White)
+	, SequenceLineThickness(1.f)
+	, MajorGridLineColor(FLinearColor::Black)
+	, MinorGridLineColor(FLinearColor(0.f, 0.f, 0.f, 0.5f))
+	, ZeroCrossingLineColor(FLinearColor::Black)
+	, ZeroCrossingLineThickness(1.f)
+	, SampleMarkersSize(2.5f)
+	, SequenceBackgroundColor(FLinearColor(0.02f, 0.02f, 0.02f, 1.f))
+	, BackgroundBrush(*FAppStyle::GetBrush(TEXT("WhiteBrush")))
+	, DesiredWidth(1280.f)
+	, DesiredHeight(720.f)
+{
+}
+
+const FSampledSequenceViewerStyle& FSampledSequenceViewerStyle::GetDefault()
+{
+	static FSampledSequenceViewerStyle Default;
+	return Default;
+}
+
+void FSampledSequenceViewerStyle::GetResources(TArray< const FSlateBrush* >& OutBrushes) const
+{
+	OutBrushes.Add(&BackgroundBrush);
+}
+
diff --git a/Engine/Plugins/Runtime/AudioWidgets/Source/AudioWidgets/Private/SSampledSequenceViewer.cpp b/Engine/Plugins/Runtime/AudioWidgets/Source/AudioWidgets/Private/SSampledSequenceViewer.cpp
new file mode 100644
--- /dev/null
+++ b/Engine/Plugins/Runtime/AudioWidgets/Source/AudioWidgets/Private/SSampledSequenceViewer.cpp
@@ -0,0 +1,264 @@
+// Copyright Epic Games, Inc. All Rights Reserved.
+
+#include "SSampledSequenceViewer.h"
+
+void SSampledSequenceViewer::Construct(const FArguments& InArgs, TArrayView<const float> InSampleData, const uint8 InNumChannels, TSharedRef<ISampledSequenceGridService> InGridService)
+{
+	GridService = InGridService;
+	UpdateView(InSampleData, InNumChannels);
+
+	Style = InArgs._Style;
+	DrawingParams = InArgs._SequenceDrawingParams;
+
+	check(Style);
+	SequenceColor = Style->SequenceColor;
+	MajorGridLineColor = Style->MajorGridLineColor;
+	MinorGridLineColor = Style->MinorGridLineColor;
+	BackgroundColor = Style->SequenceBackgroundColor;
+	BackgroundBrush = Style->BackgroundBrush;
+	DesiredWidth = Style->DesiredWidth;
+	DesiredHeight = Style->DesiredHeight;
+	SampleMarkersSize = Style->SampleMarkersSize;
+	SequenceLineThickness = Style->SequenceLineThickness;
+	ZeroCrossingLineColor = Style->ZeroCrossingLineColor;
+	ZeroCrossingLineThickness = Style->ZeroCrossingLineThickness;
+}
+
+void SSampledSequenceViewer::UpdateView(TArrayView<const float> InSampleData, const uint8 InNumChannels)
+{
+	UpdateGridMetrics();
+
+	SampleData = InSampleData;
+	NumChannels = InNumChannels;
+
+	bForceRedraw = true;
+}
+
+int32 SSampledSequenceViewer::OnPaint(const FPaintArgs& Args, const FGeometry& AllottedGeometry, const FSlateRect& MyCullingRect, FSlateWindowElementList& OutDrawElements, int32 LayerId, const FWidgetStyle& InWidgetStyle, bool bParentEnabled) const 
+{
+	float PixelWidth = MyCullingRect.GetSize().X;
+
+	FSlateDrawElement::MakeBox(
+		OutDrawElements,
+		++LayerId,
+		AllottedGeometry.ToPaintGeometry(),
+		&BackgroundBrush,
+		ESlateDrawEffect::None, 
+		BackgroundColor.GetSpecifiedColor()
+	);
+
+	if (PixelWidth > 0)
+	{
+		DrawGridLines(AllottedGeometry, OutDrawElements, LayerId);
+
+		if (SequenceDrawMode == ESequenceDrawMode::BinnedPeaks)
+		{
+			TArray<FVector2D> BinDrawPoints;
+			BinDrawPoints.SetNumUninitialized(2);
+
+			for (const TimeSeriesDrawingUtils::FSampleBinCoordinates& BinCoordinates : CachedBinsDrawCoordinates)
+			{
+				BinDrawPoints[0] = BinCoordinates.Top;
+				BinDrawPoints[1] = BinCoordinates.Bottom;
+
+				FSlateDrawElement::MakeLines(
+					OutDrawElements,
+					LayerId,
+					AllottedGeometry.ToPaintGeometry(),
+					BinDrawPoints,
+					ESlateDrawEffect::None,
+					SequenceColor.GetSpecifiedColor()
+				);
+			}
+
+		}
+		else if (SequenceDrawMode == ESequenceDrawMode::SequenceLine)
+		{
+			TArray<FVector2D> SamplesDrawCoordinates;
+			const uint16 NChannels = NumChannels;
+
+			for (uint16 Channel = 0; Channel < NChannels; ++Channel)
+			{
+				for (int32 SampleIndex = Channel; SampleIndex < CachedSampleDrawCoordinates.Num(); SampleIndex += NChannels)
+				{
+					SamplesDrawCoordinates.Add(CachedSampleDrawCoordinates[SampleIndex]);
+					
+					const FVector2D SampleBoxSize = FVector2D(SampleMarkersSize, SampleMarkersSize);
+					const FVector2D SampleBoxTransform = FVector2D(CachedSampleDrawCoordinates[SampleIndex] - (SampleBoxSize / 2.f));
+
+					const FPaintGeometry SampleBoxGeometry = AllottedGeometry.ToPaintGeometry(
+						SampleBoxSize,
+						FSlateLayoutTransform(SampleBoxTransform)
+					);
+
+					FSlateDrawElement::MakeBox(
+						OutDrawElements,
+						LayerId,
+						SampleBoxGeometry,
+						&BackgroundBrush,
+						ESlateDrawEffect::None,
+						SequenceColor.GetSpecifiedColor()
+					);
+				}
+
+				FSlateDrawElement::MakeLines(
+					OutDrawElements,
+					LayerId,
+					AllottedGeometry.ToPaintGeometry(),
+					SamplesDrawCoordinates,
+					ESlateDrawEffect::None,
+					SequenceColor.GetSpecifiedColor(),
+					true,
+					SequenceLineThickness
+				);
+
+				SamplesDrawCoordinates.Empty();
+			}
+		}
+
+		++LayerId;
+	}
+
+	return LayerId;
+}
+
+FVector2D SSampledSequenceViewer::ComputeDesiredSize(float) const
+{
+	return FVector2D(DesiredWidth, DesiredHeight);
+}
+
+void SSampledSequenceViewer::Tick(const FGeometry& AllottedGeometry, const double InCurrentTime, const float InDeltaTime)
+{
+	uint32 DiscretePixelWidth = FMath::FloorToInt(AllottedGeometry.GetLocalSize().X);
+	if (DiscretePixelWidth <= 0)
+	{
+		return;
+	}
+
+	if (DiscretePixelWidth != CachedPixelWidth || bForceRedraw)
+	{
+		CachedPixelWidth = DiscretePixelWidth;
+		CachedPixelHeight = AllottedGeometry.GetLocalSize().Y;
+		const float NumSamplesToDisplay = SampleData.Num() / (float) NumChannels;
+		SequenceDrawMode = NumSamplesToDisplay > CachedPixelWidth ? ESequenceDrawMode::BinnedPeaks : ESequenceDrawMode::SequenceLine;
+
+		if (SequenceDrawMode == ESequenceDrawMode::BinnedPeaks)
+		{
+			TimeSeriesDrawingUtils::GroupInterleavedSampledTSIntoMinMaxBins<float>(CachedPeaks, CachedPixelWidth, SampleData.GetData(), SampleData.Num(), GridMetrics.SampleRate, NumChannels);
+			TimeSeriesDrawingUtils::GenerateSampleBinsCoordinatesForGeometry(CachedBinsDrawCoordinates, AllottedGeometry, CachedPeaks, NumChannels, DrawingParams);
+		}
+		else
+		{
+			TimeSeriesDrawingUtils::GenerateSequencedSamplesCoordinatesForGeometry(CachedSampleDrawCoordinates, SampleData, AllottedGeometry, NumChannels, GridMetrics, DrawingParams);
+		}
+		
+		bForceRedraw = false;
+	}
+	else if (CachedPixelHeight != AllottedGeometry.GetLocalSize().Y)
+	{
+		CachedPixelHeight = AllottedGeometry.GetLocalSize().Y;
+
+		if (SequenceDrawMode == ESequenceDrawMode::BinnedPeaks)
+		{
+			TimeSeriesDrawingUtils::GenerateSampleBinsCoordinatesForGeometry(CachedBinsDrawCoordinates, AllottedGeometry, CachedPeaks, NumChannels, DrawingParams);
+		}
+		else
+		{
+			TimeSeriesDrawingUtils::GenerateSequencedSamplesCoordinatesForGeometry(CachedSampleDrawCoordinates, SampleData, AllottedGeometry, NumChannels, GridMetrics, DrawingParams);
+		}
+	}
+}
+
+void SSampledSequenceViewer::DrawGridLines(const FGeometry& AllottedGeometry, FSlateWindowElementList& OutDrawElements, int32& LayerId) const
+{
+	TArray<FVector2D> LinePoints;
+	LinePoints.SetNumUninitialized(2);
+
+	const double MinorGridXStep = GridMetrics.MajorGridXStep / GridMetrics.NumMinorGridDivisions;
+
+	for (double CurrentMajorLineX = GridMetrics.FirstMajorTickX; CurrentMajorLineX < AllottedGeometry.Size.X; CurrentMajorLineX += GridMetrics.MajorGridXStep)
+	{
+		const double MajorLineX = CurrentMajorLineX;
+
+		LinePoints[0] = FVector2D(MajorLineX, 0.f);
+		LinePoints[1] = FVector2D(MajorLineX, AllottedGeometry.Size.Y);
+
+		FSlateDrawElement::MakeLines(
+			OutDrawElements,
+			++LayerId,
+			AllottedGeometry.ToPaintGeometry(),
+			LinePoints,
+			ESlateDrawEffect::None,
+			MajorGridLineColor.GetSpecifiedColor(),
+			false);
+
+
+		for (int32 MinorLineIndex = 1; MinorLineIndex < GridMetrics.NumMinorGridDivisions; ++MinorLineIndex)
+		{
+			const double MinorLineX = MajorLineX + MinorGridXStep * MinorLineIndex;
+
+			LinePoints[0] = FVector2D(MinorLineX, 0.);
+			LinePoints[1] = FVector2D(MinorLineX, AllottedGeometry.Size.Y);
+
+			FSlateDrawElement::MakeLines(
+				OutDrawElements,
+				++LayerId,
+				AllottedGeometry.ToPaintGeometry(),
+				LinePoints,
+				ESlateDrawEffect::None,
+				MinorGridLineColor.GetSpecifiedColor(),
+				false);
+
+		}
+	}
+
+	const uint16 NChannels = NumChannels;
+
+	for (uint16 Channel = 0; Channel < NChannels; ++Channel)
+	{
+		const TimeSeriesDrawingUtils::FHorizontalDimensionSlot ChannelBoundaries(Channel, NChannels, AllottedGeometry);
+
+		LinePoints[0] = FVector2D(0.f, ChannelBoundaries.Center);
+		LinePoints[1] = FVector2D(AllottedGeometry.Size.X, ChannelBoundaries.Center);
+
+		FSlateDrawElement::MakeLines(
+			OutDrawElements,
+			++LayerId,
+			AllottedGeometry.ToPaintGeometry(),
+			LinePoints,
+			ESlateDrawEffect::None,
+			ZeroCrossingLineColor.GetSpecifiedColor(),
+			false,
+			ZeroCrossingLineThickness);
+
+	}
+}
+
+void SSampledSequenceViewer::UpdateGridMetrics()
+{
+	check(GridService)
+	GridMetrics = GridService->GetGridMetrics();
+}
+
+void SSampledSequenceViewer::OnStyleUpdated(const FNotifyingAudioWidgetStyle* UpdatedStyle)
+{
+	check(UpdatedStyle);
+	check(Style);
+
+	if (UpdatedStyle != Style)
+	{
+		return;
+	}
+
+	SequenceColor = Style->SequenceColor;
+	MajorGridLineColor = Style->MajorGridLineColor;
+	MinorGridLineColor = Style->MinorGridLineColor;
+	ZeroCrossingLineColor = Style->ZeroCrossingLineColor;
+	ZeroCrossingLineThickness = Style->ZeroCrossingLineThickness;
+	BackgroundColor = Style->SequenceBackgroundColor;
+	BackgroundBrush = Style->BackgroundBrush;
+	SampleMarkersSize = Style->SampleMarkersSize;
+	SequenceLineThickness = Style->SequenceLineThickness;
+	DesiredWidth = Style->DesiredWidth;
+	DesiredHeight = Style->DesiredHeight;
+}
diff --git a/Engine/Plugins/Runtime/AudioWidgets/Source/AudioWidgets/Private/TimeSeriesDrawingUtils.cpp b/Engine/Plugins/Runtime/AudioWidgets/Source/AudioWidgets/Private/TimeSeriesDrawingUtils.cpp
new file mode 100644
--- /dev/null
+++ b/Engine/Plugins/Runtime/AudioWidgets/Source/AudioWidgets/Private/TimeSeriesDrawingUtils.cpp
@@ -0,0 +1,100 @@
+// Copyright Epic Games, Inc. All Rights Reserved.
+
+#include "TimeSeriesDrawingUtils.h"
+
+#include "Layout/Geometry.h"
+
+TimeSeriesDrawingUtils::FHorizontalDimensionSlot::FHorizontalDimensionSlot(const uint16 DimensionToDraw, const uint16 TotalNumDimensions, const FGeometry& InAllottedGeometry)
+{
+	Height = InAllottedGeometry.GetLocalSize().Y / TotalNumDimensions;
+	float ChannelSlotMidPoint = Height / 2.f;
+
+	Top = Height * DimensionToDraw;
+	Center = Top + ChannelSlotMidPoint;
+	Bottom = Top + Height;
+}
+
+void TimeSeriesDrawingUtils::GenerateSampleBinsCoordinatesForGeometry(TArray<FSampleBinCoordinates>& OutDrawCoordinates, const FGeometry& InAllottedGeometry, const TArray<TRange<float>>& InSampleBins, const uint16 NDimensions, const FSampledSequenceDrawingParams Params)
+{
+	if (!ensure(NDimensions != 0))
+	{
+		return;
+	}
+
+	const uint32 PixelWidth = FMath::FloorToInt(InAllottedGeometry.GetLocalSize().X);
+	check(PixelWidth * NDimensions == InSampleBins.Num());
+
+	const float HeightScale = InAllottedGeometry.GetLocalSize().Y / (2.f * Params.MaxDisplayedValue * NDimensions) * Params.MaxSequenceHeight * Params.VerticalZoomFactor;
+
+
+	OutDrawCoordinates.SetNumUninitialized(InSampleBins.Num());
+	FSampleBinCoordinates* OutCoordinatesData = OutDrawCoordinates.GetData();
+
+	for (uint16 Channel = 0; Channel < NDimensions; ++Channel)
+	{
+		const FHorizontalDimensionSlot ChannelBoundaries(Channel, NDimensions, InAllottedGeometry);
+
+		for (uint32 Pixel = 0; Pixel < PixelWidth; ++Pixel)
+		{
+			uint32 PeakIndex = Pixel * NDimensions + Channel;
+			const TRange<float>& MinMaxBin = InSampleBins[PeakIndex];
+
+			const float SampleMaxScaled = MinMaxBin.GetUpperBoundValue() * HeightScale > Params.MinScaledBinValue ? MinMaxBin.GetUpperBoundValue() * HeightScale : Params.MinSequenceHeight;
+			const float SampleMinScaled = MinMaxBin.GetLowerBoundValue() * HeightScale < -1 * Params.MinScaledBinValue ? MinMaxBin.GetLowerBoundValue() * HeightScale : -1.f * Params.MinSequenceHeight;
+
+			const float Top = FMath::Max(ChannelBoundaries.Center - SampleMaxScaled, ChannelBoundaries.Top + Params.DimensionSlotMargin);
+			const float Bottom = FMath::Min(ChannelBoundaries.Center - SampleMinScaled, ChannelBoundaries.Bottom - Params.DimensionSlotMargin);
+
+			OutCoordinatesData[PeakIndex].Top = FVector2D(Pixel, Top);
+			OutCoordinatesData[PeakIndex].Bottom = FVector2D(Pixel, Bottom);
+		}
+	}
+}
+
+void TimeSeriesDrawingUtils::GenerateSequencedSamplesCoordinatesForGeometry(TArray<FVector2D>& OutDrawCoordinates, TArrayView<const float> InSampleData, const FGeometry& InAllottedGeometry, const uint16 NDimensions, const FSampledSequenceGridMetrics InGridMetrics, const FSampledSequenceDrawingParams Params)
+{
+	OutDrawCoordinates.Empty();
+
+	if (!ensure(NDimensions != 0))
+	{
+		return;
+	}
+
+	const uint32 NumFramesToDisplay = InSampleData.Num() / NDimensions;
+
+	TArray<FHorizontalDimensionSlot> ChannelSlotsBoundaries;
+
+	for (uint16 Channel = 0; Channel < NDimensions; ++Channel)
+	{
+		ChannelSlotsBoundaries.Emplace(Channel, NDimensions, InAllottedGeometry);
+	}
+
+	uint32 FrameIndex = 0;
+	const float XDrawMargin = InAllottedGeometry.Size.X + InGridMetrics.PixelsPerFrame;
+
+	for (double FrameX = InGridMetrics.FirstMajorTickX; FrameX < XDrawMargin; FrameX += InGridMetrics.PixelsPerFrame)
+	{
+		const double XRatio = FrameX / InAllottedGeometry.Size.X;
+
+		if (FrameX >= 0 && FrameIndex < NumFramesToDisplay)
+		{
+			for (uint16 Channel = 0; Channel < NDimensions; ++Channel)
+			{
+				const FHorizontalDimensionSlot& ChannelBoundaries = ChannelSlotsBoundaries[Channel];
+				const uint32 SampleIndex = FrameIndex * NDimensions + Channel;
+				const float SampleValue = InSampleData[SampleIndex];
+
+				const float SampleValueRatio = SampleValue / Params.MaxDisplayedValue * Params.MaxSequenceHeight * Params.VerticalZoomFactor;
+				const float TopBoundary = ChannelBoundaries.Top + Params.DimensionSlotMargin;
+				const float BottomBoundary = ChannelBoundaries.Bottom - Params.DimensionSlotMargin;
+				const float SampleY = FMath::Clamp((SampleValueRatio * ChannelBoundaries.Height / 2.f) + ChannelBoundaries.Center, TopBoundary, BottomBoundary);
+
+				const FVector2D SampleCoordinates(FrameX, SampleY);
+
+				OutDrawCoordinates.Add(SampleCoordinates);
+			}
+
+			FrameIndex++;
+		}
+	}
+}
diff --git a/Engine/Plugins/Runtime/AudioWidgets/Source/AudioWidgets/Public/AudioWidgetsSlateTypes.h b/Engine/Plugins/Runtime/AudioWidgets/Source/AudioWidgets/Public/AudioWidgetsSlateTypes.h
--- a/Engine/Plugins/Runtime/AudioWidgets/Source/AudioWidgets/Public/AudioWidgetsSlateTypes.h
+++ b/Engine/Plugins/Runtime/AudioWidgets/Source/AudioWidgets/Public/AudioWidgetsSlateTypes.h
@@ -2,10 +2,26 @@
 
 #pragma once
 
+#include "Delegates/Delegate.h"
 #include "Styling/SlateTypes.h"
 
 #include "AudioWidgetsSlateTypes.generated.h"
 
+DECLARE_MULTICAST_DELEGATE_OneParam(FOnAudioWidgetStyleUpdated, const FNotifyingAudioWidgetStyle* /*Updated Widget Style*/);
+
+USTRUCT()
+struct FNotifyingAudioWidgetStyle : public FSlateWidgetStyle
+{
+	GENERATED_USTRUCT_BODY()
+
+	FNotifyingAudioWidgetStyle() = default;
+	virtual ~FNotifyingAudioWidgetStyle() = default;
+
+	virtual void BroadcastStyleUpdate() const { OnStyleUpdated.Broadcast(this); }
+
+	FOnAudioWidgetStyleUpdated OnStyleUpdated;
+};
+
 /**
  * Represents the appearance of an Audio Text Box 
  */
@@ -170,6 +186,67 @@ struct AUDIOWIDGETS_API FAudioRadialSliderStyle : public FSlateWidgetStyle
 	FAudioRadialSliderStyle& SetDefaultSliderRadius(const float& InDefaultSliderRadius) { DefaultSliderRadius = InDefaultSliderRadius; return *this; }
 };
 
+/**
+ * Represents the appearance of a Sampled Sequence Viewer
+ */
+USTRUCT(BlueprintType)
+struct AUDIOWIDGETS_API FSampledSequenceViewerStyle : public FNotifyingAudioWidgetStyle
+{
+	GENERATED_USTRUCT_BODY()
+
+	FSampledSequenceViewerStyle();
+
+	static const FSampledSequenceViewerStyle& GetDefault();
+	virtual const FName GetTypeName() const override { return TypeName; };
+	virtual void GetResources(TArray< const FSlateBrush* >& OutBrushes) const override;
+
+	static const FName TypeName;
+
+	UPROPERTY(EditAnywhere, BlueprintReadWrite, Category = Appearance)
+	FSlateColor SequenceColor;
+	FSampledSequenceViewerStyle& SetSequenceColor(const FSlateColor InSequenceColor) { SequenceColor = InSequenceColor; BroadcastStyleUpdate(); return *this; }
+
+	UPROPERTY(EditAnywhere, BlueprintReadWrite, Category = Appearance)
+	float SequenceLineThickness;
+	FSampledSequenceViewerStyle& SetSequenceLineThickness(const float InSequenceLineThickness) { SequenceLineThickness = InSequenceLineThickness; BroadcastStyleUpdate(); return *this; }
+
+	UPROPERTY(EditAnywhere, BlueprintReadWrite, Category = Appearance)
+	FSlateColor MajorGridLineColor;
+	FSampledSequenceViewerStyle& SetMajorGridLineColor(const FSlateColor InMajorGridLineColor) { MajorGridLineColor = InMajorGridLineColor; BroadcastStyleUpdate(); return *this; }
+
+	UPROPERTY(EditAnywhere, BlueprintReadWrite, Category = Appearance)
+	FSlateColor MinorGridLineColor;
+	FSampledSequenceViewerStyle& SetMinorGridLineColor(const FSlateColor InMinorGridLineColor) { MinorGridLineColor = InMinorGridLineColor; BroadcastStyleUpdate(); return *this; }
+
+	UPROPERTY(EditAnywhere, BlueprintReadWrite, Category = Appearance)
+	FSlateColor ZeroCrossingLineColor;
+	FSampledSequenceViewerStyle& SetZeroCrossingLineColor(const FSlateColor InZeroCrossingLineColor) { ZeroCrossingLineColor = InZeroCrossingLineColor; BroadcastStyleUpdate(); return *this; }
+
+	UPROPERTY(EditAnywhere, BlueprintReadWrite, Category = Appearance)
+	float ZeroCrossingLineThickness;
+	FSampledSequenceViewerStyle& SetZeroCrossingLineThickness(const float InZeroCrossingLineThickness) { ZeroCrossingLineThickness = InZeroCrossingLineThickness; BroadcastStyleUpdate(); return *this; }
+
+	UPROPERTY(EditAnywhere, BlueprintReadWrite, Category = Appearance)
+	float SampleMarkersSize;
+	FSampledSequenceViewerStyle& SetSampleMarkersSize(const float InSampleMarkersSize) { SampleMarkersSize = InSampleMarkersSize; BroadcastStyleUpdate(); return *this; }
+
+	UPROPERTY(EditAnywhere, BlueprintReadWrite, Category = Appearance)
+	FSlateColor SequenceBackgroundColor;
+	FSampledSequenceViewerStyle& SetBackgroundColor(const FSlateColor InBackgroundColor) { SequenceBackgroundColor = InBackgroundColor; BroadcastStyleUpdate(); return *this; }
+
+	UPROPERTY(EditAnywhere, BlueprintReadWrite, Category = Appearance)
+	FSlateBrush BackgroundBrush;
+	FSampledSequenceViewerStyle& SetBackgroundBrush(const FSlateBrush InBackgroundBrush) { BackgroundBrush = InBackgroundBrush; BroadcastStyleUpdate(); return *this; }
+
+	UPROPERTY(EditAnywhere, BlueprintReadWrite, Category = Appearance)
+	float DesiredWidth;
+	FSampledSequenceViewerStyle& SetDesiredWidth(const float InDesiredWidth) { DesiredWidth = InDesiredWidth; BroadcastStyleUpdate(); return *this; }
+
+	UPROPERTY(EditAnywhere, BlueprintReadWrite, Category = Appearance)
+	float DesiredHeight;
+	FSampledSequenceViewerStyle& SetDesiredHeight(const float InDesiredHeight) { DesiredHeight = InDesiredHeight; BroadcastStyleUpdate(); return *this; }
+};
+
 #if UE_ENABLE_INCLUDE_ORDER_DEPRECATED_IN_5_2
 #include "Styling/SlateStyle.h"
 #include "Styling/SlateStyleMacros.h"
diff --git a/Engine/Plugins/Runtime/AudioWidgets/Source/AudioWidgets/Public/ISampledSequenceGridService.h b/Engine/Plugins/Runtime/AudioWidgets/Source/AudioWidgets/Public/ISampledSequenceGridService.h
new file mode 100644
--- /dev/null
+++ b/Engine/Plugins/Runtime/AudioWidgets/Source/AudioWidgets/Public/ISampledSequenceGridService.h
@@ -0,0 +1,23 @@
+// Copyright Epic Games, Inc. All Rights Reserved.
+
+#pragma once
+
+#include "Delegates/Delegate.h"
+
+struct FSampledSequenceGridMetrics
+{
+	int32 NumMinorGridDivisions = 0;
+	uint32 SampleRate = 0;
+	uint32 StartFrame = 0;
+	double PixelsPerFrame = 0;
+	double FirstMajorTickX = 0;
+	double MajorGridXStep = 0;
+};
+
+class ISampledSequenceGridService
+{
+public:
+	virtual ~ISampledSequenceGridService() = default;
+	virtual const FSampledSequenceGridMetrics GetGridMetrics() const = 0;
+	virtual const float SnapPositionToClosestFrame(const float InPixelPosition) const = 0;
+};
\ No newline at end of file
diff --git a/Engine/Plugins/Runtime/AudioWidgets/Source/AudioWidgets/Public/SSampledSequenceViewer.h b/Engine/Plugins/Runtime/AudioWidgets/Source/AudioWidgets/Public/SSampledSequenceViewer.h
new file mode 100644
--- /dev/null
+++ b/Engine/Plugins/Runtime/AudioWidgets/Source/AudioWidgets/Public/SSampledSequenceViewer.h
@@ -0,0 +1,79 @@
+// Copyright Epic Games, Inc. All Rights Reserved.
+
+#pragma once
+
+#include "AudioWidgetsSlateTypes.h"
+#include "ISampledSequenceGridService.h"
+#include "Styling/ISlateStyle.h"
+#include "Styling/SlateWidgetStyleAsset.h"
+#include "TimeSeriesDrawingUtils.h"
+#include "Widgets/SLeafWidget.h"
+
+namespace SampledSequenceViewerWidget
+{
+	typedef TRange<float> SampleRange;
+}
+
+class AUDIOWIDGETS_API SSampledSequenceViewer : public SLeafWidget
+{
+public:
+
+	SLATE_BEGIN_ARGS(SSampledSequenceViewer) 
+	{
+	}
+	
+	SLATE_ARGUMENT(TimeSeriesDrawingUtils::FSampledSequenceDrawingParams, SequenceDrawingParams)
+
+	SLATE_STYLE_ARGUMENT(FSampledSequenceViewerStyle, Style)
+		
+	SLATE_END_ARGS()
+
+	void Construct(const FArguments& InArgs, TArrayView<const float> InSampleData, const uint8 InNumChannels, TSharedRef<ISampledSequenceGridService> InGridService);
+	void UpdateView(TArrayView<const float> InSampleData, const uint8 InNumChannels);
+	
+	void OnStyleUpdated(const FNotifyingAudioWidgetStyle* UpdatedStyle);
+
+private:
+	virtual int32 OnPaint(const FPaintArgs& Args, const FGeometry& AllottedGeometry, const FSlateRect& MyCullingRect, FSlateWindowElementList& OutDrawElements, int32 LayerId, const FWidgetStyle& InWidgetStyle, bool bParentEnabled) const override;
+	virtual FVector2D ComputeDesiredSize(float) const override;
+	virtual void Tick(const FGeometry& AllottedGeometry, const double InCurrentTime, const float InDeltaTime) override;
+
+	void DrawGridLines(const FGeometry& AllottedGeometry, FSlateWindowElementList& OutDrawElements, int32& LayerId) const;
+
+	void UpdateGridMetrics();
+
+	bool bForceRedraw = false;
+
+	TArrayView<const float> SampleData;
+	uint8 NumChannels = 0;
+	float DurationInSeconds = 0.f;
+
+	TSharedPtr<ISampledSequenceGridService> GridService = nullptr;
+	FSampledSequenceGridMetrics GridMetrics;
+
+	const FSampledSequenceViewerStyle* Style = nullptr;
+	FSlateBrush BackgroundBrush;
+	FSlateColor BackgroundColor = FLinearColor(0.02f, 0.02f, 0.02f, 1.f);
+	FSlateColor SequenceColor = FLinearColor::White;
+	FSlateColor MajorGridLineColor = FLinearColor::Black;
+	FSlateColor MinorGridLineColor = FLinearColor(0.f, 0.f, 0.f, 0.5f);
+	FSlateColor ZeroCrossingLineColor = FLinearColor::Black;
+	float ZeroCrossingLineThickness = 1.f;
+	float DesiredHeight = 0.f;
+	float DesiredWidth = 0.f;
+	float SampleMarkersSize = 2.5f;
+	float SequenceLineThickness = 1.f;
+	
+
+	uint32 CachedPixelWidth = 0; 
+	float CachedPixelHeight = 0.f;
+	TArray<SampledSequenceViewerWidget::SampleRange> CachedPeaks;
+	TArray<TimeSeriesDrawingUtils::FSampleBinCoordinates> CachedBinsDrawCoordinates;
+	TimeSeriesDrawingUtils::FSampledSequenceDrawingParams DrawingParams;
+	TArray<FVector2D> CachedSampleDrawCoordinates;
+
+	enum class ESequenceDrawMode {
+		BinnedPeaks = 0,
+		SequenceLine
+	} SequenceDrawMode;
+};
\ No newline at end of file
diff --git a/Engine/Plugins/Runtime/AudioWidgets/Source/AudioWidgets/Public/TimeSeriesDrawingUtils.h b/Engine/Plugins/Runtime/AudioWidgets/Source/AudioWidgets/Public/TimeSeriesDrawingUtils.h
--- a/Engine/Plugins/Runtime/AudioWidgets/Source/AudioWidgets/Public/TimeSeriesDrawingUtils.h
+++ b/Engine/Plugins/Runtime/AudioWidgets/Source/AudioWidgets/Public/TimeSeriesDrawingUtils.h
@@ -2,8 +2,12 @@
 
 #pragma once
 
+#include "ISampledSequenceGridService.h"
+#include "Math/Vector2D.h"
 #include "Math/Range.h"
 
+struct FGeometry;
+
 namespace TimeSeriesDrawingUtils
 {
 	/**
@@ -14,39 +18,39 @@ namespace TimeSeriesDrawingUtils
 	* @param NumDesiredBins		Number of output bins
 	* @param RawDataPtr			Ptr to the beginning of the samples
 	* @param TotalNumSamples	The total number of samples of the input time series
-	* @param NChannels			Number of interleaved channels in the time series
+	* @param NDimensions		Number of interleaved dimensions in the time series
 	* @param StartRatio			The ratio of the total number of frames (in a range of 0-1) at which grouping should start
 	* @param EndRatio			The ratio of the total number of frames (in a range of 0-1) at which grouping should end.
 	*
 	*/
 	template<typename SamplesType>
-	void GroupInterleavedSamplesIntoMinMaxBins(TArray<TRange<SamplesType>>& OutBins, const uint32 NumDesiredBins, const SamplesType* RawDataPtr, const uint32 TotalNumSamples, const uint16 NChannels = 1, const double StartRatio = 0.0, double EndRatio = 1.0)
+	void GroupInterleavedSamplesIntoMinMaxBins(TArray<TRange<SamplesType>>& OutBins, const uint32 NumDesiredBins, const SamplesType* RawDataPtr, const uint32 TotalNumSamples, const uint16 NDimensions = 1, const double StartRatio = 0.0, double EndRatio = 1.0)
 	{
-		check(StartRatio >= 0.f && StartRatio < EndRatio)
-		check(EndRatio <= 1.f)
+		check(StartRatio >= 0.f && StartRatio < EndRatio);
+		check(EndRatio <= 1.f);
 
-		uint32 NumPeaks = NumDesiredBins * NChannels;
+		uint32 NumPeaks = NumDesiredBins * NDimensions;
 		OutBins.SetNumUninitialized(NumPeaks);
-		const uint32 NumFrames = TotalNumSamples / NChannels;
+		const uint32 NumFrames = TotalNumSamples / NDimensions;
 
 		double FramesPerBin = ((EndRatio - StartRatio) * NumFrames) / NumDesiredBins;
 		double IterationStartFrame = StartRatio * NumFrames;
 		uint32 FirstBinnedFrame = FMath::RoundToInt(StartRatio * NumFrames);
 
-		for (uint32 Peak = 0; Peak < NumPeaks; Peak += NChannels)
+		for (uint32 Peak = 0; Peak < NumPeaks; Peak += NDimensions)
 		{
 			uint32 LastBinnedFrame = IterationStartFrame + FramesPerBin;
 
-			for (uint16 Channel = 0; Channel < NChannels; ++Channel)
+			for (uint16 Channel = 0; Channel < NDimensions; ++Channel)
 			{
 				SamplesType MaxSampleValue = TNumericLimits<SamplesType>::Min();
 				SamplesType MinSampleValue = TNumericLimits<SamplesType>::Max();
 
 				for (uint32 Frame = FirstBinnedFrame; Frame < LastBinnedFrame; ++Frame)
 				{
-					const uint32 SampleIndex = Frame * NChannels + Channel;
+					const uint32 SampleIndex = Frame * NDimensions + Channel;
 					check(SampleIndex < TotalNumSamples);
-					int16 SampleValue = RawDataPtr[SampleIndex];
+					SamplesType SampleValue = RawDataPtr[SampleIndex];
 
 					MinSampleValue = FMath::Min(MinSampleValue, SampleValue);
 					MaxSampleValue = FMath::Max(MaxSampleValue, SampleValue);
@@ -64,27 +68,90 @@ namespace TimeSeriesDrawingUtils
 	/**
 	* Groups samples of a time series into an equal number of desired bins.
 	* Each bin contains the min and max values of the grouped samples.
-	* 
-	* @param OutBins			TArray where the bins will be written to 
+	*
+	* @param OutBins			TArray where the bins will be written to
 	* @param NumDesiredBins		Number of output bins
-	* @param RawDataPtr			Ptr to the beginning of the samples 
-	* @param TotalNumSamples	The total number of samples of the input time series 
-	* @param SampleRate			SampleRate of the time series 
-	* @param NChannels			Number of interleaved channels in the time series
+	* @param RawDataPtr			Ptr to the beginning of the samples
+	* @param TotalNumSamples	The total number of samples of the input time series
+	* @param SampleRate			SampleRate of the time series
+	* @param NDimensions		Number of interleaved dimensions in the time series
 	* @param StartTime			The time at which grouping should start
-	* @param EndTime 			The time at which grouping should end. 
-	*							
-	* Note: With a negative EndTime, the method will calculate automatically the EndTime by doing 
-	* TotalNumSamples / (SampleRate * NChannels)
-	*							
+	* @param EndTime 			The time at which grouping should end.
+	*
+	* Note: With a negative EndTime, the method will calculate automatically the EndTime by doing
+	* TotalNumSamples / (SampleRate * NDimensions)
+	*
 	*/
 	template<typename SamplesType>
-	void GroupInterleavedSampledTSIntoMinMaxBins(TArray<TRange<SamplesType>>& OutBins, const uint32 NumDesiredBins, const SamplesType* RawDataPtr, const uint32 TotalNumSamples, const uint32 SampleRate, const uint16 NChannels = 1, const float StartTime = 0.f, float EndTime = -1.f)
+	void GroupInterleavedSampledTSIntoMinMaxBins(TArray<TRange<SamplesType>>& OutBins, const uint32 NumDesiredBins, const SamplesType* RawDataPtr, const uint32 TotalNumSamples, const uint32 SampleRate, const uint16 NDimensions = 1, const float StartTime = 0.f, float EndTime = -1.f)
 	{
-		const double TotalTime = TotalNumSamples / ((float) SampleRate * NChannels);
+		const double TotalTime = TotalNumSamples / ((float)SampleRate * NDimensions);
 		const double StartRatio = StartTime / TotalTime;
 		const double EndRatio = EndTime >= 0.f ? FMath::Clamp(EndTime / TotalTime, StartRatio, 1.0) : 1.0;
 
-		GroupInterleavedSamplesIntoMinMaxBins(OutBins, NumDesiredBins, RawDataPtr, TotalNumSamples, NChannels, StartRatio, EndRatio);
+		GroupInterleavedSamplesIntoMinMaxBins(OutBins, NumDesiredBins, RawDataPtr, TotalNumSamples, NDimensions, StartRatio, EndRatio);
 	}
+
+	/**
+		* Params for drawing sequences of samples/sample bins
+		* @param MaxDisplayedValue					The highest value a sample can take
+		* @param DimensionSlotMargin				Margin to keep from the dimension slot boundaries (pixels)
+		* @param MinSequenceHeight					The minimum height (ratio) the drawn sequence can take in a channel slot
+		* @param MaxSequenceHeight					The maximum height (ratio) the drawn sequence can take in a channel slot
+		* @param MinScaledBinValue					Minimum value a scaled bin can have
+		* @param VerticalZoomFactor					Sequence zoom factor
+	*/
+	struct FSampledSequenceDrawingParams
+	{
+		double MaxDisplayedValue = 1.f;
+		float DimensionSlotMargin = 2.f;
+		float MaxSequenceHeight = 0.9f;
+		float MinSequenceHeight = 0.1f;
+		float MinScaledBinValue = 0.001f;
+		float VerticalZoomFactor = 1.f;
+	};
+
+	/**
+		* Represents an horizontal slot in which samples or sample bins can be drawn
+	*/
+	struct FHorizontalDimensionSlot
+	{
+		explicit FHorizontalDimensionSlot(const uint16 DimensionToDraw, const uint16 TotalNumDimensions, const FGeometry& InAllottedGeometry);
+
+		float Top;
+		float Center;
+		float Bottom;
+		float Height;
+	};
+
+	struct FSampleBinCoordinates
+	{
+		FVector2D Top;
+		FVector2D Bottom;
+	};
+
+	/**
+		* Generates an array of coordinates to draw binned samples as vertical lines.
+		* Coordinates will be generated for an horizontal view
+		* 
+		* @param OutDrawCoordinates			The generated coordinates
+		* @param InAllottedGeometry			The geometry to draw into
+		* @param InSampleBins				The bins to draw
+		* @param NDimensions				Number of interleaved dimensions in the time series
+		* @param Params						Drawing Params
+	*/
+	void GenerateSampleBinsCoordinatesForGeometry(TArray<FSampleBinCoordinates>& OutDrawCoordinates, const FGeometry& InAllottedGeometry, const TArray<TRange<float>>& InSampleBins, const uint16 NDimensions, const FSampledSequenceDrawingParams Params = FSampledSequenceDrawingParams());
+
+	/**
+	* Generates an array of coordinates to draw single samples in succession horizontally.
+	*
+	* @param OutDrawCoordinates			The generated coordinates
+	* @param InSampleData				The samples to draw
+	* @param InAllottedGeometry			The geometry to draw into
+	* @param NDimensions				Number of interleaved dimensions in the time series
+	* @param InGridMetrics				GridMetrics followed to lay down the samples
+	* @param Params						Drawing Params
+	*/
+	void GenerateSequencedSamplesCoordinatesForGeometry(TArray<FVector2D>& OutDrawCoordinates, TArrayView<const float> InSampleData, const FGeometry& InAllottedGeometry, const uint16 NDimensions, const FSampledSequenceGridMetrics InGridMetrics, const FSampledSequenceDrawingParams Params = FSampledSequenceDrawingParams());
+
 }
\ No newline at end of file
