[WaveformEditor] Make "Edit Waveform" asset  menu entry dynamic so it can be hidden

diff --git a/Engine/Plugins/Editor/WaveformEditor/Source/WaveformEditor/Private/WaveformEditorInstantiator.cpp b/Engine/Plugins/Editor/WaveformEditor/Source/WaveformEditor/Private/WaveformEditorInstantiator.cpp
--- a/Engine/Plugins/Editor/WaveformEditor/Source/WaveformEditor/Private/WaveformEditorInstantiator.cpp
+++ b/Engine/Plugins/Editor/WaveformEditor/Source/WaveformEditor/Private/WaveformEditorInstantiator.cpp
@@ -2,63 +2,54 @@
 
 #include "WaveformEditorInstantiator.h"
 
-#include "ContentBrowserModule.h"
+#include "Algo/AnyOf.h"
+#include "ContentBrowserMenuContexts.h"
 #include "Framework/Application/SlateApplication.h"
 #include "Framework/MultiBox/MultiBoxBuilder.h"
 #include "Sound/SoundWave.h"
+#include "Sound/SoundWaveProcedural.h"
 #include "SWaveformEditorMessageDialog.h"
+#include "ToolMenu.h"
+#include "ToolMenus.h"
+#include "ToolMenuSection.h"
 #include "WaveformEditor.h"
 #include "WaveformEditorLog.h"
 
 #define LOCTEXT_NAMESPACE "WaveformEditorInstantiator"
 
-void FWaveformEditorInstantiator::ExtendContentBrowserSelectionMenu()
+namespace WaveEditorInstantiatorPrivate
 {
-	FContentBrowserModule& ContentBrowserModule = FModuleManager::GetModuleChecked<FContentBrowserModule>(TEXT("ContentBrowser"));
-
-	TArray<FContentBrowserMenuExtender_SelectedAssets>& ContentBrowserExtenders = ContentBrowserModule.GetAllAssetViewContextMenuExtenders();
-	ContentBrowserExtenders.Add(FContentBrowserMenuExtender_SelectedAssets::CreateSP(this, &FWaveformEditorInstantiator::OnExtendContentBrowserAssetSelectionMenu));
+	auto FilterUnwantedAssets = [](const FAssetData& AssetData) { return !AssetData.IsInstanceOf<USoundWaveProcedural>(); };
 }
 
-TSharedRef<FExtender> FWaveformEditorInstantiator::OnExtendContentBrowserAssetSelectionMenu(const TArray<FAssetData>& SelectedAssets)
+void FWaveformEditorInstantiator::ExtendContentBrowserSelectionMenu()
 {
-	TSharedRef<FExtender> Extender = MakeShared<FExtender>();
-
-	Extender->AddMenuExtension(
-		"GetAssetActions",
-		EExtensionHook::After,
-		nullptr,
-		FMenuExtensionDelegate::CreateSP(this, &FWaveformEditorInstantiator::AddWaveformEditorMenuEntry, SelectedAssets)
-	);
+	UToolMenu* Menu = UToolMenus::Get()->ExtendMenu("ContentBrowser.AssetContextMenu.SoundWave");
+	FToolMenuSection& Section = Menu->FindOrAddSection("GetAssetActions");
 
-	return Extender;
-}
-
-void FWaveformEditorInstantiator::AddWaveformEditorMenuEntry(FMenuBuilder& MenuBuilder, TArray<FAssetData> SelectedAssets)
-{
-	if (SelectedAssets.Num() > 0)
-	{
-		// check that all selected assets are USoundWaves.
-		TArray<USoundWave*> SelectedSoundWaves;
-		for (const FAssetData& SelectedAsset : SelectedAssets)
+	Section.AddDynamicEntry("SoundWaveEditing_CreateWaveformEditor", FNewToolMenuSectionDelegate::CreateLambda([this](FToolMenuSection& InSection)
 		{
-			if (SelectedAsset.GetClass() != USoundWave::StaticClass())
+			if (const UContentBrowserAssetContextMenuContext* Context = InSection.FindContext<UContentBrowserAssetContextMenuContext>())
 			{
-				return;
+				if (Algo::AnyOf(Context->SelectedAssets, WaveEditorInstantiatorPrivate::FilterUnwantedAssets))
+				{
+					const TAttribute<FText> Label = LOCTEXT("SoundWave_WaveformEditor", "Edit Waveform");
+					const TAttribute<FText> ToolTip = LOCTEXT("SoundWave_WaveformEditor_Tooltip", "Open waveform editor");
+					const FSlateIcon Icon = FSlateIcon();
+					const FToolMenuExecuteAction UIAction = FToolMenuExecuteAction::CreateSP(this, &FWaveformEditorInstantiator::ExecuteCreateWaveformEditor);
+
+					InSection.AddMenuEntry("SoundWave_CreateWaveformEditor", Label, ToolTip, Icon, UIAction);
+				}
 			}
+		}));
+}
 
-			SelectedSoundWaves.Add(static_cast<USoundWave*>(SelectedAsset.GetAsset()));
-		}
-
-		MenuBuilder.AddMenuEntry(
-			LOCTEXT("SoundWave_WaveformEditor", "Edit Waveform"),
-			LOCTEXT("SoundWave_WaveformEditor_Tooltip", "Open waveform editor"),
-			FSlateIcon(),
-			FUIAction(
-				FExecuteAction::CreateSP(this, &FWaveformEditorInstantiator::CreateWaveformEditor, SelectedSoundWaves),
-				FCanExecuteAction()
-			)
-		);
+void FWaveformEditorInstantiator::ExecuteCreateWaveformEditor(const FToolMenuContext& MenuContext)
+{
+	if (const UContentBrowserAssetContextMenuContext* Context = UContentBrowserAssetContextMenuContext::FindContextWithAssets(MenuContext))
+	{
+		TArray<USoundWave*> SoundWavesToEdit = Context->LoadSelectedObjectsIf<USoundWave>(WaveEditorInstantiatorPrivate::FilterUnwantedAssets);
+		CreateWaveformEditor(SoundWavesToEdit);
 	}
 }
 
diff --git a/Engine/Plugins/Editor/WaveformEditor/Source/WaveformEditor/Private/WaveformEditorInstantiator.h b/Engine/Plugins/Editor/WaveformEditor/Source/WaveformEditor/Private/WaveformEditorInstantiator.h
--- a/Engine/Plugins/Editor/WaveformEditor/Source/WaveformEditor/Private/WaveformEditorInstantiator.h
+++ b/Engine/Plugins/Editor/WaveformEditor/Source/WaveformEditor/Private/WaveformEditorInstantiator.h
@@ -5,9 +5,8 @@
 #include "IWaveformEditorInstantiator.h"
 #include "Templates/SharedPointer.h"
 
-class FExtender;
-class FMenuBuilder;
-struct FAssetData;
+class USoundWave;
+struct FToolMenuContext;
 
 class FWaveformEditorInstantiator : public IWaveformEditorInstantiator,  public TSharedFromThis<FWaveformEditorInstantiator>
 {
@@ -17,10 +16,7 @@ public:
 
 private:
 	virtual void CreateWaveformEditor(TArray<USoundWave*> SoundWavesToEdit) override;
-
-	TSharedRef<FExtender> OnExtendContentBrowserAssetSelectionMenu(const TArray<FAssetData>& SelectedAssets);
-	void AddWaveformEditorMenuEntry(FMenuBuilder& MenuBuilder, TArray<FAssetData> SelectedAssets);
+	void ExecuteCreateWaveformEditor(const FToolMenuContext& MenuContext);
 	bool CanSoundWaveBeOpenedInEditor(const USoundWave* SoundWaveToEdit);
-
 	void DisplayErrorDialog(const FText& ErrorMessage) const;
 };
\ No newline at end of file
diff --git a/Engine/Plugins/Editor/WaveformEditor/Source/WaveformEditor/WaveformEditor.build.cs b/Engine/Plugins/Editor/WaveformEditor/Source/WaveformEditor/WaveformEditor.build.cs
--- a/Engine/Plugins/Editor/WaveformEditor/Source/WaveformEditor/WaveformEditor.build.cs
+++ b/Engine/Plugins/Editor/WaveformEditor/Source/WaveformEditor/WaveformEditor.build.cs
@@ -17,6 +17,7 @@ namespace UnrealBuildTool.Rules
 					"AssetRegistry",
 					"Core", 
 					"CoreUObject",
+					"ContentBrowser",
 					"DeveloperSettings",
 					"EditorStyle",
 					"Engine",
