[WaveformEditor]Toolset abstraction part II:
- Abstract Transport Coordinator and make indipendent from RenderData.
- Remove widgets dependencies from Transport Coordinator
- Add a generic Playhead widget
- Add generic utilities to route input through stack of widgets

diff --git a/Engine/Plugins/Editor/WaveformEditor/Source/WaveformEditor/Private/WaveformEditor.cpp b/Engine/Plugins/Editor/WaveformEditor/Source/WaveformEditor/Private/WaveformEditor.cpp
--- a/Engine/Plugins/Editor/WaveformEditor/Source/WaveformEditor/Private/WaveformEditor.cpp
+++ b/Engine/Plugins/Editor/WaveformEditor/Source/WaveformEditor/Private/WaveformEditor.cpp
@@ -6,6 +6,7 @@
 #include "Components/AudioComponent.h"
 #include "Misc/TransactionObjectEvent.h"
 #include "PropertyEditorModule.h"
+#include "SamplesSequenceTransportCoordinator.h"
 #include "Sound/SoundWave.h"
 #include "Styling/AppStyle.h"
 #include "SWaveformPanel.h"
@@ -173,7 +174,7 @@ bool FWaveformEditor::BindDelegates()
 
 	AudioComponent->OnAudioPlaybackPercentNative.AddSP(this, &FWaveformEditor::HandlePlaybackPercentageChange);
 	AudioComponent->OnAudioPlayStateChangedNative.AddSP(this, &FWaveformEditor::HandleAudioComponentPlayStateChanged);
-	TransportCoordinator->OnPlayheadScrubUpdate.AddSP(this, &FWaveformEditor::HandlePlayheadScrub);
+	TransportCoordinator->OnFocusPointScrubUpdate.AddSP(this, &FWaveformEditor::HandlePlayheadScrub);
 	return true;
 }
 
@@ -626,10 +627,10 @@ bool FWaveformEditor::SetUpWaveformPanel()
 		return false;
 	}
 
-	TSharedPtr<FWaveformEditorRenderData> RenderData = MakeShared<FWaveformEditorRenderData>();
+	WaveformRenderData = MakeShared<FWaveformEditorRenderData>();
 	
-	TransportCoordinator = MakeShared<FWaveformEditorTransportCoordinator>(RenderData.ToSharedRef());
-	RenderData->OnRenderDataUpdated.AddSP(TransportCoordinator.Get(), &FWaveformEditorTransportCoordinator::HandleRenderDataUpdate);
+	TransportCoordinator = MakeShared<FSamplesSequenceTransportCoordinator>();
+	WaveformRenderData->OnRenderDataUpdated.AddSP(this, &FWaveformEditor::HandleRenderDataUpdate);
 	
 	FOnTransformationsPropertiesRequired OnTransformationPropertiesRequired = FOnTransformationsPropertiesRequired::CreateLambda([this](FTransformationsToPropertiesArray& InObjToPropsMap)
 	{
@@ -653,7 +654,7 @@ bool FWaveformEditor::SetUpWaveformPanel()
 	});
 
 	TransformationsRenderManager = MakeShared<FWaveformTransformationsRenderManager>(SoundWave, OnTransformationPropertiesRequired);
-	TransformationsRenderManager->OnRenderDataGenerated.AddSP(RenderData.Get(), &FWaveformEditorRenderData::UpdateRenderData);
+	TransformationsRenderManager->OnRenderDataGenerated.AddSP(WaveformRenderData.Get(), &FWaveformEditorRenderData::UpdateRenderData);
 
 	TSharedPtr<SWaveformTransformationsOverlay> TransformationsOverlay = SNew(SWaveformTransformationsOverlay, TransformationsRenderManager->GetTransformLayers(), TransportCoordinator.ToSharedRef());
 
@@ -662,7 +663,7 @@ bool FWaveformEditor::SetUpWaveformPanel()
 	TransportCoordinator->OnDisplayRangeUpdated.AddSP(TransformationsOverlay.Get(), &SWaveformTransformationsOverlay::OnNewWaveformDisplayRange);
 	TransformationsRenderManager->UpdateRenderElements();
 
-	WaveformPanel = SNew(SWaveformPanel, RenderData.ToSharedRef(), TransportCoordinator.ToSharedRef(), ZoomManager.ToSharedRef(), TransformationsOverlay);
+	WaveformPanel = SNew(SWaveformPanel, WaveformRenderData.ToSharedRef(), TransportCoordinator.ToSharedRef(), ZoomManager.ToSharedRef(), TransformationsOverlay);
 
 	return WaveformPanel != nullptr;
 }
@@ -678,7 +679,7 @@ void FWaveformEditor::HandlePlaybackPercentageChange(const UAudioComponent* InCo
 		if (TransportCoordinator.IsValid())
 		{
 			const float ClampedPlayBackPercentage = FGenericPlatformMath::Fmod(InPlaybackPercentage, 1.f);
-			TransportCoordinator->ReceivePlayBackRatio(ClampedPlayBackPercentage);
+			TransportCoordinator->SetProgressRatio(ClampedPlayBackPercentage);
 		}
 	}
 }
@@ -707,7 +708,17 @@ void FWaveformEditor::HandleAudioComponentPlayStateChanged(const UAudioComponent
 	}
 }
 
-void FWaveformEditor::HandlePlayheadScrub(const uint32 SelectedSample, const uint32 TotalSampleLength, const bool bIsMoving)
+void FWaveformEditor::HandleRenderDataUpdate()
+{
+	check(WaveformRenderData != nullptr)
+
+	if (TransportCoordinator != nullptr)
+	{
+		TransportCoordinator->UpdatePlaybackRange(WaveformRenderData->GetTransformedWaveformBounds());
+	}
+}
+
+void FWaveformEditor::HandlePlayheadScrub(const float InTargetPlayBackRatio, const bool bIsMoving)
 {
 	if (bIsMoving)
 	{
@@ -719,15 +730,14 @@ void FWaveformEditor::HandlePlayheadScrub(const uint32 SelectedSample, const uin
 	}
 	else
 	{
-		const float PlayBackRatio = SelectedSample / (float)TotalSampleLength;
-		const float NewTime = PlayBackRatio * SoundWave->Duration;
+		const float NewTime = InTargetPlayBackRatio * SoundWave->Duration;
 
 		if (TransportController->IsPlaying())
 		{
 			TransportController->Seek(NewTime);
 			return;
 		}
-			
+
 		if (bWasPlayingBeforeScrubbing)
 		{
 			TransportController->Play(NewTime);
@@ -737,7 +747,7 @@ void FWaveformEditor::HandlePlayheadScrub(const uint32 SelectedSample, const uin
 		{
 			TransportController->CacheStartTime(NewTime);
 		}
-		
+
 	}
 }
 
diff --git a/Engine/Plugins/Editor/WaveformEditor/Source/WaveformEditor/Public/WaveformEditor.h b/Engine/Plugins/Editor/WaveformEditor/Source/WaveformEditor/Public/WaveformEditor.h
--- a/Engine/Plugins/Editor/WaveformEditor/Source/WaveformEditor/Public/WaveformEditor.h
+++ b/Engine/Plugins/Editor/WaveformEditor/Source/WaveformEditor/Public/WaveformEditor.h
@@ -7,7 +7,6 @@
 #include "Misc/NotifyHook.h"
 #include "Toolkits/AssetEditorToolkit.h"
 #include "WaveformEditorTransportController.h"
-#include "WaveformEditorTransportCoordinator.h"
 #include "WaveformEditorZoomController.h"
 
 class IToolkitHost;
@@ -76,7 +75,9 @@ private:
 	/** Playback delegates handlers */
 	void HandlePlaybackPercentageChange(const UAudioComponent* InComponent, const USoundWave* InSoundWave, const float InPlaybackPercentage);
 	void HandleAudioComponentPlayStateChanged(const UAudioComponent* InAudioComponent, EAudioComponentPlayState NewPlayState);
-	void HandlePlayheadScrub(const uint32 SelectedSample, const uint32 TotalSampleLength, const bool bIsMoving);
+	void HandlePlayheadScrub(const float InTargetPlayBackRatio, const bool bIsMoving);
+	
+	void HandleRenderDataUpdate();
 
 	/** FGCObject overrides */
 	virtual void AddReferencedObjects(FReferenceCollector& Collector) override;
@@ -96,11 +97,14 @@ private:
 	/** Manages render information for waveform transforms */
 	TSharedPtr<class FWaveformTransformationsRenderManager> TransformationsRenderManager = nullptr;
 
+	/** Render Data of the displayed waveform */
+	TSharedPtr<class FWaveformEditorRenderData> WaveformRenderData = nullptr;
+
 	/** Exports the edited waveform to a new asset */
 	TSharedPtr<class FWaveformEditorWaveWriter> WaveWriter = nullptr;
 
 	/** Manages Transport info in waveform panel */
-	TSharedPtr<FWaveformEditorTransportCoordinator> TransportCoordinator = nullptr;
+	TSharedPtr<class FSamplesSequenceTransportCoordinator> TransportCoordinator = nullptr;
 
 	/** Controls Transport of the audio component  */
 	TSharedPtr<FWaveformEditorTransportController> TransportController = nullptr;
diff --git a/Engine/Plugins/Editor/WaveformEditor/Source/WaveformEditor/WaveformEditor.build.cs b/Engine/Plugins/Editor/WaveformEditor/Source/WaveformEditor/WaveformEditor.build.cs
--- a/Engine/Plugins/Editor/WaveformEditor/Source/WaveformEditor/WaveformEditor.build.cs
+++ b/Engine/Plugins/Editor/WaveformEditor/Source/WaveformEditor/WaveformEditor.build.cs
@@ -13,8 +13,9 @@ namespace UnrealBuildTool.Rules
 
 			PublicDependencyModuleNames.AddRange(
 				new string[] {
-					"AudioExtensions",
 					"AssetRegistry",
+					"AudioExtensions",
+					"AudioWidgets",
 					"Core", 
 					"CoreUObject",
 					"ContentBrowser",
@@ -23,10 +24,10 @@ namespace UnrealBuildTool.Rules
 					"Engine",
 					"InputCore",
 					"SignalProcessing",
-					"UnrealEd", 
 					"Slate",
 					"SlateCore",
 					"ToolMenus",
+					"UnrealEd", 
 					"WaveformEditorWidgets",
 					"WaveformTransformationsWidgets"
 				}
diff --git a/Engine/Plugins/Editor/WaveformEditor/Source/WaveformEditorWidgets/Private/SWaveformEditorTimeRuler.cpp b/Engine/Plugins/Editor/WaveformEditor/Source/WaveformEditorWidgets/Private/SWaveformEditorTimeRuler.cpp
--- a/Engine/Plugins/Editor/WaveformEditor/Source/WaveformEditorWidgets/Private/SWaveformEditorTimeRuler.cpp
+++ b/Engine/Plugins/Editor/WaveformEditor/Source/WaveformEditorWidgets/Private/SWaveformEditorTimeRuler.cpp
@@ -5,13 +5,12 @@
 #include "Layout/WidgetPath.h"
 #include "Framework/Application/SlateApplication.h"
 #include "Framework/MultiBox/MultiBoxBuilder.h"
-#include "WaveformEditorTransportCoordinator.h"
+#include "SamplesSequenceTransportCoordinator.h"
 
 #define LOCTEXT_NAMESPACE "WaveformEditorTimeRuler"
 
-void SWaveformEditorTimeRuler::Construct(const FArguments& InArgs, TSharedRef<FWaveformEditorTransportCoordinator> InTransportCoordinator, TSharedRef<ISampledSequenceGridService> InGridService)
+void SWaveformEditorTimeRuler::Construct(const FArguments& InArgs, TSharedRef<ISampledSequenceGridService> InGridService)
 {
-	TransportCoordinator = InTransportCoordinator;
 	GridService = InGridService;
 	UpdateGridMetrics();
 
@@ -212,6 +211,11 @@ void SWaveformEditorTimeRuler::UpdateDisplayUnit(const EWaveformEditorDisplayUni
 	DisplayUnit = InDisplayUnit;
 }
 
+void SWaveformEditorTimeRuler::SetPlayheadPosition(const float InNewPosition)
+{
+	PlayheadPosition = InNewPosition;
+}
+
 void SWaveformEditorTimeRuler::OnStyleUpdated(const FWaveformEditorWidgetStyleBase* UpdatedStyle)
 {
 	check(UpdatedStyle);
@@ -238,9 +242,10 @@ void SWaveformEditorTimeRuler::OnStyleUpdated(const FWaveformEditorWidgetStyleBa
 void SWaveformEditorTimeRuler::DrawPlayheadHandle(const FGeometry& AllottedGeometry, FSlateWindowElementList& OutDrawElements, int32& LayerId) const
 {
 	const float WindowWidth = AllottedGeometry.Size.X;
-	const float SnappedHandleCenter = GridService->SnapPositionToClosestFrame(TransportCoordinator->GetPlayheadPosition() * WindowWidth);
-	const float	HandleStart = SnappedHandleCenter - HandleWidth / 2;
+	const float PlayheadCenter = PlayheadPosition;
+	const float	HandleStart = PlayheadCenter - HandleWidth / 2;
 	const float	HandleEnd = HandleStart + HandleWidth;
+
  	FPaintGeometry HandleGeometry = AllottedGeometry.ToPaintGeometry(FVector2f(HandleEnd - HandleStart, AllottedGeometry.Size.Y), FSlateLayoutTransform(FVector2f(HandleStart, 0)));
 
 	FSlateDrawElement::MakeBox(
@@ -255,28 +260,4 @@ void SWaveformEditorTimeRuler::DrawPlayheadHandle(const FGeometry& AllottedGeome
 	++LayerId;
 }
 
-FReply SWaveformEditorTimeRuler::OnMouseButtonDown(const FGeometry& MyGeometry, const FPointerEvent& MouseEvent)
-{
-	return TransportCoordinator->ReceiveMouseButtonDown(*this, MyGeometry, MouseEvent);
-}
-
-FReply SWaveformEditorTimeRuler::OnMouseButtonUp(const FGeometry& MyGeometry, const FPointerEvent& MouseEvent)
-{
-	const bool HandleRightMouseButton = MouseEvent.GetEffectingButton() == EKeys::RightMouseButton;
-
-	if (HandleRightMouseButton)
-	{
-		FWidgetPath WidgetPath = MouseEvent.GetEventPath() != nullptr ? *MouseEvent.GetEventPath() : FWidgetPath();
-		FSlateApplication::Get().PushMenu(AsShared(), WidgetPath, MakeContextMenu(), FSlateApplication::Get().GetCursorPos(), FPopupTransitionEffect::ContextMenu);
-		return FReply::Handled();
-	}
-
-	return TransportCoordinator->ReceiveMouseButtonUp(*this, MyGeometry, MouseEvent);
-}
-
-FReply SWaveformEditorTimeRuler::OnMouseMove(const FGeometry& MyGeometry, const FPointerEvent& MouseEvent)
-{
-	return TransportCoordinator->ReceiveMouseMove(*this, MyGeometry, MouseEvent);
-}
-
 #undef LOCTEXT_NAMESPACE
\ No newline at end of file
diff --git a/Engine/Plugins/Editor/WaveformEditor/Source/WaveformEditorWidgets/Private/SWaveformEditorTimeRuler.h b/Engine/Plugins/Editor/WaveformEditor/Source/WaveformEditorWidgets/Private/SWaveformEditorTimeRuler.h
--- a/Engine/Plugins/Editor/WaveformEditor/Source/WaveformEditorWidgets/Private/SWaveformEditorTimeRuler.h
+++ b/Engine/Plugins/Editor/WaveformEditor/Source/WaveformEditorWidgets/Private/SWaveformEditorTimeRuler.h
@@ -12,7 +12,7 @@
 
 DECLARE_MULTICAST_DELEGATE_OneParam(FOnTimeUnitMenuSelection, const EWaveformEditorDisplayUnit /* Requested Display Unit */);
 
-class FWaveformEditorTransportCoordinator;
+class FSamplesSequenceTransportCoordinator;
 
 class SWaveformEditorTimeRuler : public SCompoundWidget
 {
@@ -29,9 +29,10 @@ class SWaveformEditorTimeRuler : public SCompoundWidget
 	SLATE_END_ARGS()
 
 public:
-	void Construct(const FArguments& InArgs, TSharedRef<FWaveformEditorTransportCoordinator> InTransportCoordinator, TSharedRef<ISampledSequenceGridService> InGridService);
+	void Construct(const FArguments& InArgs, TSharedRef<ISampledSequenceGridService> InGridService);
 	void UpdateGridMetrics();
 	void UpdateDisplayUnit(const EWaveformEditorDisplayUnit InDisplayUnit);
+	void SetPlayheadPosition(const float InNewPosition);
 
 	void OnStyleUpdated(const FWaveformEditorWidgetStyleBase* UpdatedStyle);
 
@@ -42,10 +43,6 @@ private:
 	virtual int32 OnPaint(const FPaintArgs& Args, const FGeometry& AllottedGeometry, const FSlateRect& MyCullingRect, FSlateWindowElementList& OutDrawElements, int32 LayerId, const FWidgetStyle& InWidgetStyle, bool bParentEnabled) const override;
 	virtual FVector2D ComputeDesiredSize(float) const override;
 
-	virtual FReply OnMouseButtonDown(const FGeometry& MyGeometry, const FPointerEvent& MouseEvent) override;
-	virtual FReply OnMouseButtonUp(const FGeometry& MyGeometry, const FPointerEvent& MouseEvent) override;
-	virtual FReply OnMouseMove(const FGeometry& MyGeometry, const FPointerEvent& MouseEvent) override;
-
 	void DrawPlayheadHandle(const FGeometry& AllottedGeometry, FSlateWindowElementList& OutDrawElements, int32& LayerId) const;
 	void DrawRulerTicks(const FGeometry& AllottedGeometry, FSlateWindowElementList& OutDrawElements, int32& LayerId) const;
 	void DrawTickTimeString(uint32 TickFrame, const double TickX, const double TickY, FSlateWindowElementList& OutDrawElements, int32& LayerId, const FGeometry& AllottedGeometry) const;
@@ -68,9 +65,9 @@ private:
 	float DesiredWidth = 0.f;
 	float HandleWidth = 15.f;
 	float TicksTextOffset = 5.f;
+	float PlayheadPosition = 0.f;
 	FSlateFontInfo TicksTextFont;
 
-	TSharedPtr<FWaveformEditorTransportCoordinator> TransportCoordinator = nullptr;
 	TSharedPtr<ISampledSequenceGridService> GridService = nullptr;
 
 	EWaveformEditorDisplayUnit DisplayUnit;
diff --git a/Engine/Plugins/Editor/WaveformEditor/Source/WaveformEditorWidgets/Private/SWaveformPanel.cpp b/Engine/Plugins/Editor/WaveformEditor/Source/WaveformEditorWidgets/Private/SWaveformPanel.cpp
--- a/Engine/Plugins/Editor/WaveformEditor/Source/WaveformEditorWidgets/Private/SWaveformPanel.cpp
+++ b/Engine/Plugins/Editor/WaveformEditor/Source/WaveformEditorWidgets/Private/SWaveformPanel.cpp
@@ -2,21 +2,22 @@
 
 #include "SWaveformPanel.h"
 
+#include "SPlayheadOverlay.h"
+#include "SSampledSequenceViewer.h"
 #include "SWaveformEditorTimeRuler.h"
 #include "SWaveformTransformationsOverlay.h"
-#include "SSampledSequenceViewer.h"
 #include "SWaveformViewerOverlay.h"
+#include "SamplesSequenceTransportCoordinator.h"
 #include "WaveformEditorDisplayUnit.h"
 #include "WaveformEditorGridData.h"
 #include "WaveformEditorRenderData.h"
 #include "WaveformEditorStyle.h"
-#include "WaveformEditorTransportCoordinator.h"
 #include "WaveformEditorZoomController.h"
 #include "Widgets/SBoxPanel.h"
 #include "Widgets/SCompoundWidget.h"
 #include "Widgets/SOverlay.h"
 
-void SWaveformPanel::Construct(const FArguments& InArgs, TSharedRef<FWaveformEditorRenderData> InRenderData, TSharedRef<FWaveformEditorTransportCoordinator> InTransportCoordinator, TSharedRef<FWaveformEditorZoomController> InZoomManager, TSharedPtr<SWaveformTransformationsOverlay> InWaveformTransformationsOverlay)
+void SWaveformPanel::Construct(const FArguments& InArgs, TSharedRef<FWaveformEditorRenderData> InRenderData, TSharedRef<FSamplesSequenceTransportCoordinator> InTransportCoordinator, TSharedRef<FWaveformEditorZoomController> InZoomManager, TSharedPtr<SWaveformTransformationsOverlay> InWaveformTransformationsOverlay)
 {
 	DisplayUnit = EWaveformEditorDisplayUnit::Seconds;
 
@@ -39,11 +40,16 @@ void SWaveformPanel::Construct(const FArguments& InArgs, TSharedRef<FWaveformEdi
 		WaveformTransformationsOverlay = InWaveformTransformationsOverlay;
 	}
 	
-	SetUpWaveformViewerOverlay(InTransportCoordinator, InZoomManager);
-	SetUpTimeRuler(InTransportCoordinator, GridData.ToSharedRef());
+	SetupPlayheadOverlay();
+
+	SetUpWaveformViewerOverlay(InZoomManager);
+	SetUpTimeRuler(GridData.ToSharedRef());
 	CreateLayout();
+	
 }
 
+
+
 void SWaveformPanel::CreateLayout()
 {
 	check(TimeRuler);
@@ -64,11 +70,17 @@ void SWaveformPanel::CreateLayout()
 		];
 	}
 
+	WaveformView->AddSlot()
+	[
+		PlayheadOverlay.ToSharedRef()
+	];
+
 	WaveformView->AddSlot()
 	[
 		WaveformViewerOverlay.ToSharedRef()
 	];
 
+
 	ChildSlot
 	[
 		SNew(SVerticalBox)
@@ -83,59 +95,79 @@ void SWaveformPanel::CreateLayout()
 	];
 }
 
-void SWaveformPanel::SetUpTimeRuler(TSharedRef<FWaveformEditorTransportCoordinator> InTransportCoordinator, TSharedRef<FWaveformEditorGridData> InGridData)
+void SWaveformPanel::SetUpTimeRuler(TSharedRef<FWaveformEditorGridData> InGridData)
 {
 	FWaveformEditorTimeRulerStyle* TimeRulerStyle = &WaveformEditorStyle->GetRegisteredWidgetStyle<FWaveformEditorTimeRulerStyle>("WaveformEditorRuler.Style").Get();
 	check(TimeRulerStyle);
 
-	TimeRuler = SNew(SWaveformEditorTimeRuler, InTransportCoordinator, InGridData).DisplayUnit(DisplayUnit).Style(TimeRulerStyle);
+	TimeRuler = SNew(SWaveformEditorTimeRuler, InGridData).DisplayUnit(DisplayUnit).Style(TimeRulerStyle);
+	GridData->OnGridMetricsUpdated.AddSP(TimeRuler.ToSharedRef(), &SWaveformEditorTimeRuler::UpdateGridMetrics);
 	TimeRulerStyle->OnStyleUpdated.AddSP(TimeRuler.ToSharedRef(), &SWaveformEditorTimeRuler::OnStyleUpdated);
 	TimeRuler->OnTimeUnitMenuSelection.AddSP(this, &SWaveformPanel::UpdateDisplayUnit);
+	TimeRuler->SetOnMouseButtonDown(FPointerEventHandler::CreateSP(this, &SWaveformPanel::HandleTimeRulerMouseButtonDown));
+	TimeRuler->SetOnMouseButtonUp(FPointerEventHandler::CreateSP(this, &SWaveformPanel::HandleTimeRulerMouseButtonUp));
+	TimeRuler->SetOnMouseMove(FPointerEventHandler::CreateSP(this, &SWaveformPanel::HandleTimeRulerMouseMove));
 }
 
-void SWaveformPanel::SetUpWaveformViewerOverlay(TSharedRef<FWaveformEditorTransportCoordinator> InTransportCoordinator, TSharedRef<FWaveformEditorZoomController> InZoomManager)
+void SWaveformPanel::SetUpWaveformViewerOverlay(TSharedRef<FWaveformEditorZoomController> InZoomManager)
 {
-	FWaveformViewerOverlayStyle* OverlayStyle = &WaveformEditorStyle->GetRegisteredWidgetStyle<FWaveformViewerOverlayStyle>("WaveformViewerOverlay.Style").Get();
-	check(OverlayStyle);
+	FSampledSequenceViewerStyle* ViewerStyle = &WaveformEditorStyle->GetRegisteredWidgetStyle<FSampledSequenceViewerStyle>("WaveformViewer.Style").Get();
+	check(ViewerStyle);
+	
+	TArray<TSharedPtr<SWidget>> OverlaidWidgets;
 
-	WaveformViewerOverlay = SNew(SWaveformViewerOverlay, InTransportCoordinator, WaveformTransformationsOverlay.ToSharedRef(), GridData.ToSharedRef()).Style(OverlayStyle);
-	OverlayStyle->OnStyleUpdated.AddSP(WaveformViewerOverlay.ToSharedRef(), &SWaveformViewerOverlay::OnStyleUpdated);
+	if (WaveformTransformationsOverlay)
+	{
+		OverlaidWidgets.Add(WaveformTransformationsOverlay);
+	}
+
+	check(PlayheadOverlay)
+	OverlaidWidgets.Add(PlayheadOverlay);
+
+	WaveformViewerOverlay = SNew(SWaveformViewerOverlay, OverlaidWidgets).Style(ViewerStyle);
 	WaveformViewerOverlay->OnNewMouseDelta.BindSP(InZoomManager, &FWaveformEditorZoomController::ZoomByDelta);
 }
 
+void SWaveformPanel::SetupPlayheadOverlay()
+{
+	FPlayheadOverlayStyle* PlayheadOverlayStyle = &WaveformEditorStyle->GetRegisteredWidgetStyle<FPlayheadOverlayStyle>("WaveformEditorPlayheadOverlay.Style").Get();
+	check(PlayheadOverlayStyle);
+
+	PlayheadOverlay = SNew(SPlayheadOverlay).Style(PlayheadOverlayStyle);
+	PlayheadOverlay->SetOnMouseButtonDown(FPointerEventHandler::CreateSP(this, &SWaveformPanel::HandlePlayheadOverlayMouseButtonUp));
+	PlayheadOverlayStyle->OnStyleUpdated.AddSP(PlayheadOverlay.ToSharedRef(), &SPlayheadOverlay::OnStyleUpdated);
+}
+
 void SWaveformPanel::SetUpWaveformViewer(TSharedRef<FWaveformEditorGridData> InGridData, TSharedRef<FWaveformEditorRenderData> InRenderData)
 {
 	FSampledSequenceViewerStyle* WaveViewerStyle = &WaveformEditorStyle->GetRegisteredWidgetStyle<FSampledSequenceViewerStyle>("WaveformViewer.Style").Get();
 	check(WaveViewerStyle);
 	
-	InGridData->UpdateGridMetrics(WaveViewerStyle->DesiredWidth);
 	TimeSeriesDrawingUtils::FSampledSequenceDrawingParams WaveformViewerDrawingParams;
 	WaveformViewerDrawingParams.MaxDisplayedValue = TNumericLimits<int16>::Max();
 
 	WaveformViewer = SNew(SSampledSequenceViewer, MakeArrayView(FloatRenderData.GetData(), FloatRenderData.Num()), InRenderData->GetNumChannels(), InGridData).Style(WaveViewerStyle).SequenceDrawingParams(WaveformViewerDrawingParams);
 	WaveViewerStyle->OnStyleUpdated.AddSP(WaveformViewer.ToSharedRef(), &SSampledSequenceViewer::OnStyleUpdated);
+	GridData->OnGridMetricsUpdated.AddSP(WaveformViewer.ToSharedRef(), &SSampledSequenceViewer::UpdateGridMetrics);
 }
 
 void SWaveformPanel::SetUpGridData(TSharedRef<FWaveformEditorRenderData> InRenderData)
 {
-	GridData = MakeShared<FWaveformEditorGridData>(InRenderData->GetNumSamples() / InRenderData->GetNumChannels(), InRenderData->GetSampleRate());
-
-	const ISlateStyle* WaveEditorStyle = FSlateStyleRegistry::FindSlateStyle("WaveformEditorStyle");
-	if (ensure(WaveEditorStyle))
-	{
-		const FWaveformEditorTimeRulerStyle& RulerStyle = WaveEditorStyle->GetWidgetStyle<FWaveformEditorTimeRulerStyle>("WaveformEditorRuler.Style");
-		GridData->SetTicksTimeFont(&RulerStyle.TicksTextFont);
-	}
+	FWaveformEditorTimeRulerStyle* RulerStyle = &WaveformEditorStyle->GetRegisteredWidgetStyle<FWaveformEditorTimeRulerStyle>("WaveformEditorRuler.Style").Get();
+	check(RulerStyle)
+	
+	GridData = MakeShared<FWaveformEditorGridData>(InRenderData->GetNumSamples() / InRenderData->GetNumChannels(), InRenderData->GetSampleRate(), RulerStyle->DesiredWidth, &RulerStyle->TicksTextFont);
 }
 
-void SWaveformPanel::SetUpZoomManager(TSharedRef<FWaveformEditorZoomController> InZoomManager, TSharedRef<FWaveformEditorTransportCoordinator> InTransportCoordinator)
+void SWaveformPanel::SetUpZoomManager(TSharedRef<FWaveformEditorZoomController> InZoomManager, TSharedRef<FSamplesSequenceTransportCoordinator> InTransportCoordinator)
 {
-	InZoomManager->OnZoomRatioChanged.AddSP(InTransportCoordinator, &FWaveformEditorTransportCoordinator::OnZoomLevelChanged);
+	InZoomManager->OnZoomRatioChanged.AddSP(InTransportCoordinator, &FSamplesSequenceTransportCoordinator::SetZoomRatio);
 }
 
 void SWaveformPanel::OnRenderDataUpdated()
 {
 	check(TransportCoordinator)
+	GenerateFloatRenderData();
 	OnDisplayRangeUpdated(TransportCoordinator->GetDisplayRange());
 }
 
@@ -160,11 +192,6 @@ void SWaveformPanel::OnDisplayRangeUpdated(const TRange<float> NewDisplayRange)
 		GridData->UpdateDisplayRange(TRange<uint32>(FirstRenderedFrame, FirstRenderedFrame + NumFramesToRender));
 	}
 
-	if (TimeRuler)
-	{
-		TimeRuler->UpdateGridMetrics();
-	}
-
 	if (WaveformViewer)
 	{
 		TArrayView<const float> RenderedView = MakeArrayView(FloatRenderData.GetData(), FloatRenderData.Num());
@@ -175,6 +202,8 @@ void SWaveformPanel::OnDisplayRangeUpdated(const TRange<float> NewDisplayRange)
 
 void SWaveformPanel::GenerateFloatRenderData()
 {
+	check(RenderData)
+
 	TArrayView<const int16> SampleData = RenderData->GetSampleData();
 	FloatRenderData.SetNumUninitialized(SampleData.Num());
 
@@ -201,6 +230,101 @@ void SWaveformPanel::Tick(const FGeometry& AllottedGeometry, const double InCurr
 		if (GridData)
 		{
 			GridData->UpdateGridMetrics(PaintedWidth);
-		}	
+		}
 	}
-}
\ No newline at end of file
+
+	UpdatePlayheadPosition(PaintedWidth);
+
+}
+
+void SWaveformPanel::UpdatePlayheadPosition(const float PaintedWidth)
+{
+	float PlayheadX = 0.f;
+
+	check(TransportCoordinator)
+	PlayheadX = TransportCoordinator->GetFocusPoint() * PaintedWidth;
+	PlayheadX = GridData ? GridData->SnapPositionToClosestFrame(PlayheadX) : PlayheadX;
+
+
+	if (PlayheadOverlay)
+	{
+		PlayheadOverlay->SetPlayheadPosition(PlayheadX);
+	}
+
+	if (TimeRuler)
+	{
+		TimeRuler->SetPlayheadPosition(PlayheadX);
+	}
+}
+
+FReply SWaveformPanel::HandleTimeRulerMouseButtonDown(const FGeometry& Geometry, const FPointerEvent& MouseEvent)
+{
+	return HandleTimeRulerInteraction(WaveformEditorPanel::EReceivedInteractionType::MouseButtonDown, MouseEvent, Geometry);
+}
+
+FReply SWaveformPanel::HandleTimeRulerMouseButtonUp(const FGeometry& Geometry, const FPointerEvent& MouseEvent)
+{
+	return HandleTimeRulerInteraction(WaveformEditorPanel::EReceivedInteractionType::MouseButtonUp, MouseEvent, Geometry);
+}
+
+FReply SWaveformPanel::HandleTimeRulerMouseMove(const FGeometry& Geometry, const FPointerEvent& MouseEvent)
+{
+	return HandleTimeRulerInteraction(WaveformEditorPanel::EReceivedInteractionType::MouseMove, MouseEvent, Geometry);
+}
+
+FReply SWaveformPanel::HandleTimeRulerInteraction(const WaveformEditorPanel::EReceivedInteractionType MouseInteractionType, const FPointerEvent& MouseEvent, const FGeometry& Geometry)
+{
+	check(TimeRuler)
+	check(TransportCoordinator)
+
+	const float LocalWidth = Geometry.GetLocalSize().X;
+
+	if (LocalWidth > 0.f)
+	{
+		const FVector2D ScreenSpacePosition = MouseEvent.GetScreenSpacePosition();
+		const FVector2D CursorPosition = Geometry.AbsoluteToLocal(ScreenSpacePosition);
+		const float CursorXRatio = CursorPosition.X / LocalWidth;
+
+		switch (MouseInteractionType)
+		{
+		default:
+			break;
+		case WaveformEditorPanel::EReceivedInteractionType::MouseButtonDown:
+			if (MouseEvent.IsMouseButtonDown(EKeys::LeftMouseButton))
+			{
+				return FReply::Handled().CaptureMouse(TimeRuler.ToSharedRef()).PreventThrottling();
+			}
+		case WaveformEditorPanel::EReceivedInteractionType::MouseMove:
+			if (MouseEvent.IsMouseButtonDown(EKeys::LeftMouseButton))
+			{
+				TransportCoordinator->ScrubFocusPoint(CursorXRatio, true);
+				return FReply::Handled().CaptureMouse(TimeRuler.ToSharedRef());
+			}
+		case WaveformEditorPanel::EReceivedInteractionType::MouseButtonUp:
+			if (TimeRuler->HasMouseCapture())
+			{
+				TransportCoordinator->ScrubFocusPoint(CursorXRatio, false);
+				return FReply::Handled().ReleaseMouseCapture();
+			}
+		}
+	}
+
+	return FReply::Handled();
+}
+
+FReply SWaveformPanel::HandlePlayheadOverlayMouseButtonUp(const FGeometry& Geometry, const FPointerEvent& MouseEvent)
+{
+	const bool HandleLeftMouseButton = MouseEvent.GetEffectingButton() == EKeys::LeftMouseButton;
+	if (HandleLeftMouseButton && TransportCoordinator)
+	{
+		const float LocalWidth = Geometry.GetLocalSize().X;
+
+		if (LocalWidth > 0.f)
+		{
+			const float NewPosition = Geometry.AbsoluteToLocal(MouseEvent.GetScreenSpacePosition()).X / LocalWidth;
+			TransportCoordinator->ScrubFocusPoint(NewPosition, false);
+		}
+	}
+
+	return FReply::Handled();
+}
diff --git a/Engine/Plugins/Editor/WaveformEditor/Source/WaveformEditorWidgets/Private/SWaveformTransformationsOverlay.cpp b/Engine/Plugins/Editor/WaveformEditor/Source/WaveformEditorWidgets/Private/SWaveformTransformationsOverlay.cpp
--- a/Engine/Plugins/Editor/WaveformEditor/Source/WaveformEditorWidgets/Private/SWaveformTransformationsOverlay.cpp
+++ b/Engine/Plugins/Editor/WaveformEditor/Source/WaveformEditorWidgets/Private/SWaveformTransformationsOverlay.cpp
@@ -2,11 +2,12 @@
 
 #include "SWaveformTransformationsOverlay.h"
 
-#include "Widgets/SOverlay.h"
-#include "WaveformEditorTransportCoordinator.h"
+#include "AudioWidgetsUtils.h"
 #include "SWaveformTransformationRenderLayer.h"
+#include "SamplesSequenceTransportCoordinator.h"
+#include "Widgets/SOverlay.h"
 
-void SWaveformTransformationsOverlay::Construct(const FArguments& InArgs, TArrayView<const FTransformationLayerRenderInfo> InTransformationRenderers, TSharedRef<FWaveformEditorTransportCoordinator> InTransportCoordinator)
+void SWaveformTransformationsOverlay::Construct(const FArguments& InArgs, TArrayView<const FTransformationLayerRenderInfo> InTransformationRenderers, TSharedRef<FSamplesSequenceTransportCoordinator> InTransportCoordinator)
 {
 	TransformationRenderers = InTransformationRenderers;
 	TransportCoordinator = InTransportCoordinator;
@@ -29,7 +30,7 @@ void SWaveformTransformationsOverlay::CreateLayout()
 	for (int32 i = 0; i < NumLayers; ++i)
 	{
 		const FTransformationLayerRenderInfo& LayerRenderInfo = TransformationRenderers[i];
-
+		
 		if (LayerRenderInfo.Key)
 		{
 			float LeftAnchor = TransportCoordinator->ConvertAbsoluteRatioToZoomed(LayerRenderInfo.Value.Key);
@@ -61,7 +62,6 @@ void SWaveformTransformationsOverlay::CreateLayout()
 	}
 }
 
-
 void SWaveformTransformationsOverlay::UpdateAnchors()
 {
 	check(TransformationRenderers.Num() == LayersSlots.Num())
@@ -80,29 +80,6 @@ void SWaveformTransformationsOverlay::UpdateAnchors()
 	}
 }
 
-FReply SWaveformTransformationsOverlay::RouteMouseInput(WidgetMouseInputFunction InputFunction, const FGeometry& MyGeometry, const FPointerEvent& MouseEvent)
-{
-	FReply TransformationInteraction = FReply::Unhandled();
-
-	//input is propagated from topmost transformation to lowermost
-	for (int32 i = TransformationLayers.Num() - 1; i >= 0; i--)
-	{
-		const TSharedPtr<SWaveformTransformationRenderLayer> LayerWidget = TransformationLayers[i];
-
-		if (LayerWidget)
-		{
-			TransformationInteraction = (LayerWidget.Get()->*InputFunction)(LayerWidget->GetTickSpaceGeometry(), MouseEvent);
-
-			if (TransformationInteraction.IsEventHandled())
-			{
-				return TransformationInteraction;
-			}
-		}
-	}
-		
-	return TransformationInteraction;
-}
-
 void SWaveformTransformationsOverlay::OnLayerChainGenerated(FTransformationLayerRenderInfo* FirstLayerPtr, const int32 NLayers)
 {
 	TransformationRenderers = MakeArrayView(FirstLayerPtr, NLayers);
@@ -121,47 +98,25 @@ void SWaveformTransformationsOverlay::OnNewWaveformDisplayRange(const TRange<flo
 
 FReply SWaveformTransformationsOverlay::OnMouseButtonDown(const FGeometry& MyGeometry, const FPointerEvent& MouseEvent)
 {
-	return RouteMouseInput(&SWidget::OnMouseButtonDown, MyGeometry, MouseEvent);
+	return AudioWidgetsUtils::RouteMouseInput(&SWidget::OnMouseButtonDown, MouseEvent, MakeArrayView(TransformationLayers.GetData(), TransformationLayers.Num()), true);
 }
 
 FReply SWaveformTransformationsOverlay::OnMouseButtonUp(const FGeometry& MyGeometry, const FPointerEvent& MouseEvent)
 {
-	return RouteMouseInput(&SWidget::OnMouseButtonUp, MyGeometry, MouseEvent);
+	return AudioWidgetsUtils::RouteMouseInput(&SWidget::OnMouseButtonUp, MouseEvent, MakeArrayView(TransformationLayers.GetData(), TransformationLayers.Num()), true);
 }
 
 FReply SWaveformTransformationsOverlay::OnMouseMove(const FGeometry& MyGeometry, const FPointerEvent& MouseEvent)
 {
-	return RouteMouseInput(&SWidget::OnMouseMove, MyGeometry, MouseEvent);
+	return AudioWidgetsUtils::RouteMouseInput(&SWidget::OnMouseMove, MouseEvent, MakeArrayView(TransformationLayers.GetData(), TransformationLayers.Num()), true);
 }
 
 FReply SWaveformTransformationsOverlay::OnMouseWheel(const FGeometry& MyGeometry, const FPointerEvent& MouseEvent)
 {
-	return RouteMouseInput(&SWidget::OnMouseWheel, MyGeometry, MouseEvent);
+	return AudioWidgetsUtils::RouteMouseInput(&SWidget::OnMouseWheel, MouseEvent, MakeArrayView(TransformationLayers.GetData(), TransformationLayers.Num()), true);
 }
 
 FCursorReply SWaveformTransformationsOverlay::OnCursorQuery(const FGeometry& MyGeometry, const FPointerEvent& CursorEvent) const
 {
-	return RouteCursorQuery(MyGeometry, CursorEvent);
-}
-
-FCursorReply SWaveformTransformationsOverlay::RouteCursorQuery(const FGeometry& MyGeometry, const FPointerEvent& CursorEvent) const
-{
-	FCursorReply CursorReply = FCursorReply::Unhandled();
-
-	for (int32 i = TransformationLayers.Num() - 1; i >= 0; i--)
-	{
-		const TSharedPtr<SWaveformTransformationRenderLayer> LayerWidget = TransformationLayers[i];
-
-		if (LayerWidget)
-		{
-			CursorReply = LayerWidget->OnCursorQuery(LayerWidget->GetTickSpaceGeometry(), CursorEvent);
-
-			if (CursorReply.IsEventHandled())
-			{
-				return CursorReply;
-			}
-		}
-	}
-
-	return CursorReply;
+	return AudioWidgetsUtils::RouteCursorQuery(CursorEvent, MakeArrayView(TransformationLayers.GetData(), TransformationLayers.Num()), true);
 }
\ No newline at end of file
diff --git a/Engine/Plugins/Editor/WaveformEditor/Source/WaveformEditorWidgets/Private/SWaveformViewerOverlay.cpp b/Engine/Plugins/Editor/WaveformEditor/Source/WaveformEditorWidgets/Private/SWaveformViewerOverlay.cpp
--- a/Engine/Plugins/Editor/WaveformEditor/Source/WaveformEditorWidgets/Private/SWaveformViewerOverlay.cpp
+++ b/Engine/Plugins/Editor/WaveformEditor/Source/WaveformEditorWidgets/Private/SWaveformViewerOverlay.cpp
@@ -2,41 +2,39 @@
 
 #include "SWaveformViewerOverlay.h"
 
-#include "ISampledSequenceGridService.h"
-#include "WaveformEditorTransportCoordinator.h"
-#include "WaveformEditorGridData.h"
+#include "AudioWidgetsUtils.h"
 #include "Widgets/SLeafWidget.h"
 
-void SWaveformViewerOverlay::Construct(const FArguments& InArgs, TSharedRef<FWaveformEditorTransportCoordinator> InTransportCoordinator, TSharedRef<SWaveformTransformationsOverlay> InTransformationsOverlay, TSharedRef<ISampledSequenceGridService> InGridService)
+void SWaveformViewerOverlay::Construct(const FArguments& InArgs, const TArray<TSharedPtr<SWidget>>& InOverlaidWidgets)
 {
-	TransportCoordinator = InTransportCoordinator;
-	TransformationsOverlay = InTransformationsOverlay;
-	GridService = InGridService;
-
-	Style = InArgs._Style;
+	if (InArgs._Style)
+	{
+		Style = InArgs._Style;
+		DesiredWidth = Style->DesiredWidth;
+		DesiredHeight = Style->DesiredHeight;
+	}
 	
-	check(Style);
-	PlayheadColor = Style->PlayheadColor;
-	PlayheadWidth = Style->PlayheadWidth;
-	DesiredWidth = Style->DesiredWidth;
-	DesiredHeight = Style->DesiredHeight;
+	OverlaidWidgets = InOverlaidWidgets;
 }
 
 FReply SWaveformViewerOverlay::OnMouseButtonDown(const FGeometry& MyGeometry, const FPointerEvent& MouseEvent)
 {
-	FReply InteractionReply = TransformationsOverlay->OnMouseButtonDown(MyGeometry, MouseEvent);
+	return AudioWidgetsUtils::RouteMouseInput(&SWidget::OnMouseButtonDown, MouseEvent, MakeArrayView(OverlaidWidgets.GetData(), OverlaidWidgets.Num()));
+}
 
-	if (!InteractionReply.IsEventHandled())
-	{
-		InteractionReply = TransportCoordinator->ReceiveMouseButtonDown(*this, MyGeometry, MouseEvent);
-	}
-	
-	return InteractionReply;
+FReply SWaveformViewerOverlay::OnMouseButtonUp(const FGeometry& MyGeometry, const FPointerEvent& MouseEvent)
+{
+	return AudioWidgetsUtils::RouteMouseInput(&SWidget::OnMouseButtonUp, MouseEvent, MakeArrayView(OverlaidWidgets.GetData(), OverlaidWidgets.Num()));
+}
+
+FReply SWaveformViewerOverlay::OnMouseMove(const FGeometry& MyGeometry, const FPointerEvent& MouseEvent)
+{
+	return AudioWidgetsUtils::RouteMouseInput(&SWidget::OnMouseMove, MouseEvent, MakeArrayView(OverlaidWidgets.GetData(), OverlaidWidgets.Num()));
 }
 
 FReply SWaveformViewerOverlay::OnMouseWheel(const FGeometry& MyGeometry, const FPointerEvent& MouseEvent)
 {
-	FReply InteractionReply = TransformationsOverlay->OnMouseWheel(MyGeometry, MouseEvent);
+	FReply InteractionReply = AudioWidgetsUtils::RouteMouseInput(&SWidget::OnMouseWheel, MouseEvent, MakeArrayView(OverlaidWidgets.GetData(), OverlaidWidgets.Num()));
 
 	if (!InteractionReply.IsEventHandled())
 	{
@@ -48,61 +46,13 @@ FReply SWaveformViewerOverlay::OnMouseWheel(const FGeometry& MyGeometry, const F
 
 }
 
-FReply SWaveformViewerOverlay::OnMouseMove(const FGeometry& MyGeometry, const FPointerEvent& MouseEvent)
-{
-	return TransformationsOverlay->OnMouseMove(MyGeometry, MouseEvent);
-}
-
 FCursorReply SWaveformViewerOverlay::OnCursorQuery(const FGeometry& MyGeometry, const FPointerEvent& CursorEvent) const
 {
-	return TransformationsOverlay->OnCursorQuery(MyGeometry, CursorEvent);
-}
-
-void SWaveformViewerOverlay::OnStyleUpdated(const FWaveformEditorWidgetStyleBase* UpdatedStyle)
-{
-	check(UpdatedStyle);
-	check(Style);
-
-	if (UpdatedStyle != Style)
-	{
-		return;
-	}
-
-	PlayheadColor = Style->PlayheadColor;
-	PlayheadWidth = Style->PlayheadWidth;
-	DesiredWidth = Style->DesiredWidth;
-	DesiredHeight = Style->DesiredHeight;
+	return AudioWidgetsUtils::RouteCursorQuery(CursorEvent, MakeArrayView(OverlaidWidgets.GetData(), OverlaidWidgets.Num()));
 }
 
 int32 SWaveformViewerOverlay::OnPaint(const FPaintArgs& Args, const FGeometry& AllottedGeometry, const FSlateRect& MyCullingRect, FSlateWindowElementList& OutDrawElements, int32 LayerId, const FWidgetStyle& InWidgetStyle, bool bParentEnabled) const
 {
-	LayerId = DrawPlayhead(AllottedGeometry, OutDrawElements, LayerId);
-
-	return LayerId;
-}
-
-int32 SWaveformViewerOverlay::DrawPlayhead(const FGeometry& AllottedGeometry, FSlateWindowElementList& OutDrawElements, int32 LayerId) const
-{
-	const float PlayHeadX = GridService->SnapPositionToClosestFrame(AllottedGeometry.GetLocalSize().X * TransportCoordinator->GetPlayheadPosition());
-
-	TArray<FVector2D> LinePoints;
-	{
-		LinePoints.AddUninitialized(2);
-		LinePoints[0] = FVector2D(PlayHeadX, 0.0f);
-		LinePoints[1] = FVector2D(PlayHeadX, AllottedGeometry.Size.Y);
-	}
-
-	FSlateDrawElement::MakeLines(
-		OutDrawElements,
-		LayerId,
-		AllottedGeometry.ToPaintGeometry(),
-		LinePoints,
-		ESlateDrawEffect::None,
-		PlayheadColor.GetSpecifiedColor(),
-		true, 
-		PlayheadWidth
-	);
-
 	return ++LayerId;
 }
 
diff --git a/Engine/Plugins/Editor/WaveformEditor/Source/WaveformEditorWidgets/Private/SWaveformViewerOverlay.h b/Engine/Plugins/Editor/WaveformEditor/Source/WaveformEditorWidgets/Private/SWaveformViewerOverlay.h
--- a/Engine/Plugins/Editor/WaveformEditor/Source/WaveformEditorWidgets/Private/SWaveformViewerOverlay.h
+++ b/Engine/Plugins/Editor/WaveformEditor/Source/WaveformEditorWidgets/Private/SWaveformViewerOverlay.h
@@ -2,17 +2,15 @@
 
 #pragma once
 
-#include "Styling/SlateStyleRegistry.h"
+#include "AudioWidgetsSlateTypes.h"
+#include "Styling/ISlateStyle.h"
 #include "Styling/SlateWidgetStyleAsset.h"
-#include "WaveformEditorSlateTypes.h"
-#include "WaveformEditorTransportCoordinator.h"
-#include "SWaveformTransformationsOverlay.h"
 #include "Widgets/SLeafWidget.h"
 
 DECLARE_DELEGATE_OneParam(FOnNewMouseDelta, const float /* new delta */)
 
-class ISampledSequenceGridService;
 class FWaveformEditorGridData;
+class FSamplesSequenceTransportCoordinator;
 
 class SWaveformViewerOverlay : public SLeafWidget
 {
@@ -21,35 +19,28 @@ public:
 	{
 	}
 
-	SLATE_STYLE_ARGUMENT(FWaveformViewerOverlayStyle, Style)
+	SLATE_STYLE_ARGUMENT(FSampledSequenceViewerStyle, Style)
 
 	SLATE_END_ARGS()
 
-	void Construct(const FArguments& InArgs, TSharedRef<FWaveformEditorTransportCoordinator> InTransportCoordinator, TSharedRef<SWaveformTransformationsOverlay> InTransformationsOverlay, TSharedRef<ISampledSequenceGridService> InGridService);
+	void Construct(const FArguments& InArgs, const TArray<TSharedPtr<SWidget>>& InOverlaidWidgets);
 
 	virtual FReply OnMouseButtonDown(const FGeometry& MyGeometry, const FPointerEvent& MouseEvent) override;
+	virtual FReply OnMouseButtonUp(const FGeometry& MyGeometry, const FPointerEvent& MouseEvent)override;
 	virtual FReply OnMouseWheel(const FGeometry& MyGeometry, const FPointerEvent& MouseEvent) override;
-	FReply OnMouseMove(const FGeometry& MyGeometry, const FPointerEvent& MouseEvent) override;
-	FCursorReply OnCursorQuery(const FGeometry& MyGeometry, const FPointerEvent& CursorEvent) const override;
+	virtual FReply OnMouseMove(const FGeometry& MyGeometry, const FPointerEvent& MouseEvent) override;
+	virtual FCursorReply OnCursorQuery(const FGeometry& MyGeometry, const FPointerEvent& CursorEvent) const override;
 
-	void OnStyleUpdated(const FWaveformEditorWidgetStyleBase* UpdatedStyle);
-
-	FOnNewMouseDelta OnNewMouseDelta;
+	FOnNewMouseDelta OnNewMouseDelta;\
 
 private:
 	virtual int32 OnPaint(const FPaintArgs& Args, const FGeometry& AllottedGeometry, const FSlateRect& MyCullingRect, FSlateWindowElementList& OutDrawElements, int32 LayerId, const FWidgetStyle& InWidgetStyle, bool bParentEnabled) const override;
 	virtual FVector2D ComputeDesiredSize(float) const override;
 
-	int32 DrawPlayhead(const FGeometry& AllottedGeometry, FSlateWindowElementList& OutDrawElements, int32 LayerId) const;
-
-	TSharedPtr<FWaveformEditorTransportCoordinator> TransportCoordinator = nullptr;
-	TSharedPtr<SWaveformTransformationsOverlay> TransformationsOverlay = nullptr;
-	TSharedPtr<ISampledSequenceGridService> GridService = nullptr;
+	TArray<TSharedPtr<SWidget>> OverlaidWidgets;
 
-	const FWaveformViewerOverlayStyle* Style = nullptr;
+	const FSampledSequenceViewerStyle* Style = nullptr;
 
-	FSlateColor PlayheadColor = FLinearColor(255.f, 0.1f, 0.2f, 1.f);
-	float PlayheadWidth = 1.0;
-	float DesiredWidth = 0.f;
-	float DesiredHeight = 0.f;
+	float DesiredWidth = 1280.f;
+	float DesiredHeight = 720.f;
 };
\ No newline at end of file
diff --git a/Engine/Plugins/Editor/WaveformEditor/Source/WaveformEditorWidgets/Private/WaveformEditorGridData.cpp b/Engine/Plugins/Editor/WaveformEditor/Source/WaveformEditorWidgets/Private/WaveformEditorGridData.cpp
--- a/Engine/Plugins/Editor/WaveformEditor/Source/WaveformEditorWidgets/Private/WaveformEditorGridData.cpp
+++ b/Engine/Plugins/Editor/WaveformEditor/Source/WaveformEditorWidgets/Private/WaveformEditorGridData.cpp
@@ -6,32 +6,34 @@
 #include "Framework/Application/SlateApplication.h"
 #include "GenericPlatform/GenericPlatformMath.h"
 
-FWaveformEditorGridData::FWaveformEditorGridData(const uint32 InTotalFrames, const uint32 InSampleRateHz, const FSlateFontInfo* InTicksTimeFont /*= nullptr*/)
- 	: TotalFrames(InTotalFrames)
+FWaveformEditorGridData::FWaveformEditorGridData(const uint32 InTotalFrames, const uint32 InSampleRateHz, const float InGridSizePixels /*= 1.f*/, const FSlateFontInfo* InTicksTimeFont /*= nullptr*/)
+	: TotalFrames(InTotalFrames)
 	, DisplayRange(TRange<uint32>::Inclusive(0, InTotalFrames))
+	, GridSizePixels(InGridSizePixels)
  	, TicksTimeFont(InTicksTimeFont)
  	, GridFrameRate(InSampleRateHz, 1)
 {
+	UpdateGridMetrics(GridSizePixels);
 }
 
 void FWaveformEditorGridData::UpdateDisplayRange(const TRange<uint32> InDisplayRange)
 {
 	check(DisplayRange.Size<uint32>() >= 2)
 	DisplayRange = InDisplayRange;
-	UpdateGridMetrics(GridPixelWidth);
+	UpdateGridMetrics(GridSizePixels);
 }
 
-bool FWaveformEditorGridData::UpdateGridMetrics(const float InGridPixelWidth)
+bool FWaveformEditorGridData::UpdateGridMetrics(const float InGridSizePixels)
 {
-	if (GridPixelWidth != InGridPixelWidth)
+	if (GridSizePixels != InGridSizePixels)
 	{
-		GridPixelWidth = InGridPixelWidth;
+		GridSizePixels = InGridSizePixels;
 	}
 
 	const float WaveformDurationSeconds = TotalFrames / (double) GridFrameRate.Numerator;
 	const double StartTimeSeconds = DisplayRange.GetLowerBoundValue() / (double) GridFrameRate.Numerator;	
 	const double DisplayedDuration = (DisplayRange.Size<double>() - 1) / GridFrameRate.Numerator; //we account for one less frame so we can distribute them evenly on the screen
-	const double PixelsPerSecond = GridPixelWidth / DisplayedDuration;
+	const double PixelsPerSecond = GridSizePixels / DisplayedDuration;
 	double MajorGridStepSeconds = 0.0;
 	float MinTicksPixelSpacing = 30.0f;
 	
@@ -49,7 +51,7 @@ bool FWaveformEditorGridData::UpdateGridMetrics(const float InGridPixelWidth)
 		return false;
 	}
 
-	const bool bZeroGridOffset = DisplayRange.Size<uint32>() < GridPixelWidth;
+	const bool bZeroGridOffset = DisplayRange.Size<uint32>() < GridSizePixels;
 	const double GridOffset = bZeroGridOffset ? 0 : FGenericPlatformMath::Fmod(StartTimeSeconds, MajorGridStepSeconds);
 	GridMetrics.FirstMajorTickX = (0.f - GridOffset) * PixelsPerSecond;
 	GridMetrics.MajorGridXStep = MajorGridStepSeconds * PixelsPerSecond;
@@ -57,6 +59,11 @@ bool FWaveformEditorGridData::UpdateGridMetrics(const float InGridPixelWidth)
 	GridMetrics.PixelsPerFrame =  PixelsPerSecond / GridFrameRate.Numerator;
 	GridMetrics.SampleRate = GridFrameRate.Numerator;
 
+	if (OnGridMetricsUpdated.IsBound())
+	{
+		OnGridMetricsUpdated.Broadcast();
+	}
+
 	return true;
 }
 
diff --git a/Engine/Plugins/Editor/WaveformEditor/Source/WaveformEditorWidgets/Private/WaveformEditorSlateTypes.cpp b/Engine/Plugins/Editor/WaveformEditor/Source/WaveformEditorWidgets/Private/WaveformEditorSlateTypes.cpp
--- a/Engine/Plugins/Editor/WaveformEditor/Source/WaveformEditorWidgets/Private/WaveformEditorSlateTypes.cpp
+++ b/Engine/Plugins/Editor/WaveformEditor/Source/WaveformEditorWidgets/Private/WaveformEditorSlateTypes.cpp
@@ -15,23 +15,6 @@ namespace WaveformEditorStylesSharedParams
 	const float ViewerWidth = 1280.f;
 }
 
-const FName FWaveformViewerOverlayStyle::TypeName("FWaveformViewerOverlayStyle");
-
-FWaveformViewerOverlayStyle::FWaveformViewerOverlayStyle()
-	: PlayheadColor(WaveformEditorStylesSharedParams::PlayheadColor)
-	, PlayheadWidth(1.0f)
-	, DesiredWidth(WaveformEditorStylesSharedParams::ViewerWidth)
-	, DesiredHeight(WaveformEditorStylesSharedParams::ViewerHeight)
-{
-}
-
-const FWaveformViewerOverlayStyle& FWaveformViewerOverlayStyle::GetDefault()
-{
-	static FWaveformViewerOverlayStyle Default;
-	return Default;
-}
-
-
 const FName FWaveformEditorTimeRulerStyle::TypeName("FWaveformEditorTimeRulerStyle");
 
 FWaveformEditorTimeRulerStyle::FWaveformEditorTimeRulerStyle()
diff --git a/Engine/Plugins/Editor/WaveformEditor/Source/WaveformEditorWidgets/Private/WaveformEditorStyle.cpp b/Engine/Plugins/Editor/WaveformEditor/Source/WaveformEditorWidgets/Private/WaveformEditorStyle.cpp
--- a/Engine/Plugins/Editor/WaveformEditor/Source/WaveformEditorWidgets/Private/WaveformEditorStyle.cpp
+++ b/Engine/Plugins/Editor/WaveformEditor/Source/WaveformEditorWidgets/Private/WaveformEditorStyle.cpp
@@ -9,7 +9,7 @@
 #include "WaveformEditorWidgetsSettings.h"
 
 static FLazyName WaveformViewerStyleName("WaveformViewer.Style");
-static FLazyName WaveformViewerOverlayStyleName("WaveformViewerOverlay.Style");
+static FLazyName PlayheadOverlayStyleName("WaveformEditorPlayheadOverlay.Style");
 static FLazyName WaveformEditorRulerStyleName("WaveformEditorRuler.Style");
 
 FName FWaveformEditorStyle::StyleName("WaveformEditorStyle");
@@ -58,9 +58,9 @@ void FWaveformEditorStyle::Init()
 
 	StyleInstance->Set(WaveformViewerStyleName, WaveViewerStyle);
 
-	//Waveform Viewer Overlay style
-	FWaveformViewerOverlayStyle WaveViewerOverlayStyle = FWaveformViewerOverlayStyle().SetPlayheadColor(Settings->PlayheadColor);
-	StyleInstance->Set(WaveformViewerOverlayStyleName, WaveViewerOverlayStyle);
+	//Playhead Overlay style
+	FPlayheadOverlayStyle PlayheadOverlayStyle = FPlayheadOverlayStyle().SetPlayheadColor(Settings->PlayheadColor);
+	StyleInstance->Set(PlayheadOverlayStyleName, PlayheadOverlayStyle);
 
 	//Time Ruler style 
 	FWaveformEditorTimeRulerStyle TimeRulerStyle = FWaveformEditorTimeRulerStyle()
@@ -87,7 +87,7 @@ const UWaveformEditorWidgetsSettings* FWaveformEditorStyle::GetWidgetsSettings()
 void FWaveformEditorStyle::OnWidgetSettingsUpdated(const FName& PropertyName, const UWaveformEditorWidgetsSettings* Settings)
 {
 	TSharedRef<FSampledSequenceViewerStyle> WaveformViewerStyle = GetRegisteredWidgetStyle<FSampledSequenceViewerStyle>(WaveformViewerStyleName);
-	TSharedRef<FWaveformViewerOverlayStyle> WaveformViewerOverlayStyle = GetRegisteredWidgetStyle<FWaveformViewerOverlayStyle>(WaveformViewerOverlayStyleName);
+	TSharedRef<FPlayheadOverlayStyle> PlayheadOverlayStyle = GetRegisteredWidgetStyle<FPlayheadOverlayStyle>(PlayheadOverlayStyleName);
 	TSharedRef<FWaveformEditorTimeRulerStyle> WaveformEditorTimeRulerStyle = GetRegisteredWidgetStyle<FWaveformEditorTimeRulerStyle>(WaveformEditorRulerStyleName);
 
 	if (PropertyName == GET_MEMBER_NAME_CHECKED(UWaveformEditorWidgetsSettings, WaveformColor))
@@ -124,7 +124,7 @@ void FWaveformEditorStyle::OnWidgetSettingsUpdated(const FName& PropertyName, co
 	}
 	else if (PropertyName == GET_MEMBER_NAME_CHECKED(UWaveformEditorWidgetsSettings, PlayheadColor))
 	{
-		WaveformViewerOverlayStyle->SetPlayheadColor(Settings->PlayheadColor);
+		PlayheadOverlayStyle->SetPlayheadColor(Settings->PlayheadColor);
 		WaveformEditorTimeRulerStyle->SetHandleColor(Settings->PlayheadColor);
 	}
 	else if (PropertyName == GET_MEMBER_NAME_CHECKED(UWaveformEditorWidgetsSettings, RulerBackgroundColor))
diff --git a/Engine/Plugins/Editor/WaveformEditor/Source/WaveformEditorWidgets/Private/WaveformEditorTransportCoordinator.cpp b/Engine/Plugins/Editor/WaveformEditor/Source/WaveformEditorWidgets/Private/WaveformEditorTransportCoordinator.cpp
deleted file mode 100644
diff --git a/Engine/Plugins/Editor/WaveformEditor/Source/WaveformEditorWidgets/Public/SWaveformPanel.h b/Engine/Plugins/Editor/WaveformEditor/Source/WaveformEditorWidgets/Public/SWaveformPanel.h
--- a/Engine/Plugins/Editor/WaveformEditor/Source/WaveformEditorWidgets/Public/SWaveformPanel.h
+++ b/Engine/Plugins/Editor/WaveformEditor/Source/WaveformEditorWidgets/Public/SWaveformPanel.h
@@ -5,16 +5,27 @@
 #include "HAL/Platform.h"
 #include "Widgets/SCompoundWidget.h"
 
+enum class EWaveformEditorDisplayUnit;
+class FSamplesSequenceTransportCoordinator;
 class FWaveformEditorGridData;
 class FWaveformEditorRenderData;
 class FWaveformEditorStyle;
-class FWaveformEditorTransportCoordinator;
 class FWaveformEditorZoomController;
+class SPlayheadOverlay;
+class SSampledSequenceViewer;
 class SWaveformEditorTimeRuler;
 class SWaveformTransformationsOverlay;
-class SSampledSequenceViewer;
 class SWaveformViewerOverlay;
-enum class EWaveformEditorDisplayUnit;
+
+namespace WaveformEditorPanel
+{
+	enum class EReceivedInteractionType
+	{
+		MouseButtonUp,
+		MouseButtonDown,
+		MouseMove
+	};
+}
 
 class WAVEFORMEDITORWIDGETS_API SWaveformPanel : public SCompoundWidget
 {
@@ -24,40 +35,49 @@ public:
 		SLATE_DEFAULT_SLOT(FArguments, InArgs)
 	SLATE_END_ARGS()
 
-	void Construct(const FArguments& InArgs, TSharedRef<FWaveformEditorRenderData> InRenderData, TSharedRef<FWaveformEditorTransportCoordinator> InTransportCoordinator, TSharedRef<FWaveformEditorZoomController> InZoomManager, TSharedPtr<SWaveformTransformationsOverlay> InWaveformTransformationsOverlay = nullptr);
+	void Construct(const FArguments& InArgs, TSharedRef<FWaveformEditorRenderData> InRenderData, TSharedRef<FSamplesSequenceTransportCoordinator> InTransportCoordinator, TSharedRef<FWaveformEditorZoomController> InZoomManager, TSharedPtr<SWaveformTransformationsOverlay> InWaveformTransformationsOverlay = nullptr);
 	virtual void Tick(const FGeometry& AllottedGeometry, const double InCurrentTime, const float InDeltaTime) override;
 
-private: 
+private:
 	void CreateLayout();
 
-	void SetUpWaveformViewer(TSharedRef<FWaveformEditorGridData> InGridData, TSharedRef<FWaveformEditorRenderData> InRenderData);
-	void SetUpWaveformViewerOverlay(TSharedRef<FWaveformEditorTransportCoordinator> InTransportCoordinator, TSharedRef<FWaveformEditorZoomController> InZoomManager);
-	void SetUpTimeRuler(TSharedRef<FWaveformEditorTransportCoordinator> InTransportCoordinator, TSharedRef<FWaveformEditorGridData> InGridData);
 	void SetUpGridData(TSharedRef<FWaveformEditorRenderData> InRenderData);
-	void SetUpZoomManager(TSharedRef<FWaveformEditorZoomController> InZoomManager, TSharedRef<FWaveformEditorTransportCoordinator> InTransportCoordinator);
+	void SetupPlayheadOverlay();
+	void SetUpTimeRuler(TSharedRef<FWaveformEditorGridData> InGridData);
+	void SetUpWaveformViewer(TSharedRef<FWaveformEditorGridData> InGridData, TSharedRef<FWaveformEditorRenderData> InRenderData);
+	void SetUpWaveformViewerOverlay(TSharedRef<FWaveformEditorZoomController> InZoomManager);
+	void SetUpZoomManager(TSharedRef<FWaveformEditorZoomController> InZoomManager, TSharedRef<FSamplesSequenceTransportCoordinator> InTransportCoordinator);
 
 	void OnRenderDataUpdated();
 	void OnDisplayRangeUpdated(const TRange<float> NewDisplayRange);
 
-	//temporary float render data conversion while widget are made generic
+	//child widgets input overrides
+	FReply HandleTimeRulerMouseButtonUp(const FGeometry& Geometry, const FPointerEvent& MouseEvent);
+	FReply HandleTimeRulerMouseMove(const FGeometry& Geometry, const FPointerEvent& MouseEvent);
+	FReply HandleTimeRulerMouseButtonDown(const FGeometry& Geometry, const FPointerEvent& MouseEvent);
+	FReply HandleTimeRulerInteraction(const WaveformEditorPanel::EReceivedInteractionType MouseInteractionType, const FPointerEvent& MouseEvent, const FGeometry& Geometry);
+	FReply HandlePlayheadOverlayMouseButtonUp(const FGeometry& Geometry, const FPointerEvent& MouseEvent);
+	
+	//temporary float render data conversion while widgets are made generic
 	void GenerateFloatRenderData();
 	TArray<float> FloatRenderData;
 
 	void UpdateDisplayUnit(const EWaveformEditorDisplayUnit InDisplayUnit);
+	void UpdatePlayheadPosition(const float PaintedWidth);
 
-	TSharedPtr<FWaveformEditorRenderData> RenderData = nullptr;
-	TSharedPtr<FWaveformEditorGridData> GridData = nullptr;
-	TSharedPtr<FWaveformEditorTransportCoordinator> TransportCoordinator = nullptr;
+	TSharedPtr<FWaveformEditorRenderData> RenderData;
+	TSharedPtr<FWaveformEditorGridData> GridData;
+	TSharedPtr<FSamplesSequenceTransportCoordinator> TransportCoordinator;
 
-	TSharedPtr<SWaveformEditorTimeRuler> TimeRuler = nullptr;
-	TSharedPtr<SSampledSequenceViewer> WaveformViewer = nullptr;
-	TSharedPtr<SWaveformTransformationsOverlay> WaveformTransformationsOverlay = nullptr;
-	TSharedPtr<SWaveformViewerOverlay> WaveformViewerOverlay = nullptr;
+	TSharedPtr<SWaveformEditorTimeRuler> TimeRuler;
+	TSharedPtr<SSampledSequenceViewer> WaveformViewer;
+	TSharedPtr<SWaveformTransformationsOverlay> WaveformTransformationsOverlay;
+	TSharedPtr<SWaveformViewerOverlay> WaveformViewerOverlay;
+	TSharedPtr<SPlayheadOverlay> PlayheadOverlay;
 
 	float CachedPixelWidth = 0.f;
 
 	EWaveformEditorDisplayUnit DisplayUnit;
 
 	FWaveformEditorStyle* WaveformEditorStyle;
-
 };
\ No newline at end of file
diff --git a/Engine/Plugins/Editor/WaveformEditor/Source/WaveformEditorWidgets/Public/SWaveformTransformationsOverlay.h b/Engine/Plugins/Editor/WaveformEditor/Source/WaveformEditorWidgets/Public/SWaveformTransformationsOverlay.h
--- a/Engine/Plugins/Editor/WaveformEditor/Source/WaveformEditorWidgets/Public/SWaveformTransformationsOverlay.h
+++ b/Engine/Plugins/Editor/WaveformEditor/Source/WaveformEditorWidgets/Public/SWaveformTransformationsOverlay.h
@@ -5,7 +5,7 @@
 #include "Widgets/SCompoundWidget.h"
 #include "Widgets/Layout/SConstraintCanvas.h"
 
-class FWaveformEditorTransportCoordinator;
+class FSamplesSequenceTransportCoordinator;
 class IWaveformTransformationRenderer;
 class SOverlay;
 class SWaveformTransformationRenderLayer;
@@ -20,7 +20,7 @@ public:
 		SLATE_DEFAULT_SLOT(FArguments, InArgs)
 	SLATE_END_ARGS()
 
-	void Construct(const FArguments& InArgs, TArrayView< const FTransformationLayerRenderInfo> InTransformationRenderers, TSharedRef<FWaveformEditorTransportCoordinator> InTransportCoordinator);
+	void Construct(const FArguments& InArgs, TArrayView< const FTransformationLayerRenderInfo> InTransformationRenderers, TSharedRef<FSamplesSequenceTransportCoordinator> InTransportCoordinator);
 	void OnLayerChainGenerated(FTransformationLayerRenderInfo* FirstLayerPtr, const int32 NLayers);
 	void UpdateLayerConstraints();
 	void OnNewWaveformDisplayRange(const TRange<float> NewDisplayRange);
@@ -36,13 +36,8 @@ private:
 	void UpdateAnchors();
 	
 	TSharedPtr<SOverlay> MainOverlayPtr;
-	TArray<TSharedPtr<SWaveformTransformationRenderLayer>> TransformationLayers;
+	TArray<TSharedPtr<SWidget>> TransformationLayers;
 	TArray<SConstraintCanvas::FSlot*> LayersSlots;
 	TArrayView<const FTransformationLayerRenderInfo> TransformationRenderers;
-	TSharedPtr<FWaveformEditorTransportCoordinator> TransportCoordinator = nullptr;
-
-	typedef FReply(SWidget::*WidgetMouseInputFunction)(const FGeometry&, const FPointerEvent&);
-	
-	FReply RouteMouseInput(WidgetMouseInputFunction InputFunction, const FGeometry& MyGeometry, const FPointerEvent& MouseEvent);
-	FCursorReply RouteCursorQuery(const FGeometry& MyGeometry, const FPointerEvent& CursorEvent) const; 
+	TSharedPtr<FSamplesSequenceTransportCoordinator> TransportCoordinator = nullptr;
 };
\ No newline at end of file
diff --git a/Engine/Plugins/Editor/WaveformEditor/Source/WaveformEditorWidgets/Public/WaveformEditorGridData.h b/Engine/Plugins/Editor/WaveformEditor/Source/WaveformEditorWidgets/Public/WaveformEditorGridData.h
--- a/Engine/Plugins/Editor/WaveformEditor/Source/WaveformEditorWidgets/Public/WaveformEditorGridData.h
+++ b/Engine/Plugins/Editor/WaveformEditor/Source/WaveformEditorWidgets/Public/WaveformEditorGridData.h
@@ -8,25 +8,29 @@
 #include "Math/Range.h"
 #include "Misc/FrameRate.h"
 
+DECLARE_MULTICAST_DELEGATE(FOnWaveformEditorGridUpdated)
+
 struct FSlateFontInfo;
 
 class WAVEFORMEDITORWIDGETS_API FWaveformEditorGridData : public ISampledSequenceGridService
 {
 public: 
-	explicit FWaveformEditorGridData(const uint32 InTotalFrames, const uint32 InSampleRateHz, const FSlateFontInfo* InTicksTimeFont = nullptr);
+	explicit FWaveformEditorGridData(const uint32 InTotalFrames, const uint32 InSampleRateHz, const float InGridSizePixels = 1.f, const FSlateFontInfo* InTicksTimeFont = nullptr);
 
 	void UpdateDisplayRange(const TRange<uint32> InDisplayRange);
-	bool UpdateGridMetrics(const float InGridPixelWidth);
+	bool UpdateGridMetrics(const float InGridSizePixels);
 	virtual const FSampledSequenceGridMetrics GetGridMetrics() const override;
 	void SetTicksTimeFont(const FSlateFontInfo* InNewFont);
 	const float SnapPositionToClosestFrame(const float InPixelPosition) const;
 
+	FOnWaveformEditorGridUpdated OnGridMetricsUpdated;
+
 private:
 	FSampledSequenceGridMetrics GridMetrics;
 	uint32 TotalFrames = 0;
 	TRange<uint32> DisplayRange;
 
-	float GridPixelWidth = 0.f;
+	float GridSizePixels = 0.f;
 	const FSlateFontInfo* TicksTimeFont = nullptr;
 	FFrameRate GridFrameRate;
 };
\ No newline at end of file
diff --git a/Engine/Plugins/Editor/WaveformEditor/Source/WaveformEditorWidgets/Public/WaveformEditorSlateTypes.h b/Engine/Plugins/Editor/WaveformEditor/Source/WaveformEditorWidgets/Public/WaveformEditorSlateTypes.h
--- a/Engine/Plugins/Editor/WaveformEditor/Source/WaveformEditorWidgets/Public/WaveformEditorSlateTypes.h
+++ b/Engine/Plugins/Editor/WaveformEditor/Source/WaveformEditorWidgets/Public/WaveformEditorSlateTypes.h
@@ -23,38 +23,6 @@ struct FWaveformEditorWidgetStyleBase : public FSlateWidgetStyle
 	FOnWidgetStyleUpdated OnStyleUpdated;
 };
 
-/**
- * Represents the appearance of a Waveform Viewer Overlay style
- */
-USTRUCT(BlueprintType)
-struct FWaveformViewerOverlayStyle : public FWaveformEditorWidgetStyleBase
-{
-	GENERATED_USTRUCT_BODY()
-
-	FWaveformViewerOverlayStyle();
-
-	static const FWaveformViewerOverlayStyle& GetDefault();
-	virtual const FName GetTypeName() const override { return TypeName; };
-
-	static const FName TypeName;
-
-	UPROPERTY(EditAnywhere, BlueprintReadWrite, Category = Appearance)
-	FSlateColor PlayheadColor;
-	FWaveformViewerOverlayStyle& SetPlayheadColor(const FSlateColor InPlayheadColor) { PlayheadColor = InPlayheadColor; BroadcastStyleUpdate(); return *this; }
-
-	UPROPERTY(EditAnywhere, BlueprintReadWrite, Category = Appearance)
-	float PlayheadWidth;
-	FWaveformViewerOverlayStyle& SetPlayheadWidth(const float InPlayheadWidth) { PlayheadWidth = InPlayheadWidth; BroadcastStyleUpdate(); return *this; }
-
-	UPROPERTY(EditAnywhere, BlueprintReadWrite, Category = Appearance)
-	float DesiredWidth;
-	FWaveformViewerOverlayStyle& SetDesiredWidth(const float InDesiredWidth) { DesiredWidth = InDesiredWidth; BroadcastStyleUpdate(); return *this; }
-
-	UPROPERTY(EditAnywhere, BlueprintReadWrite, Category = Appearance)
-	float DesiredHeight;
-	FWaveformViewerOverlayStyle& SetDesiredHeight(const float InDesiredHeight) { DesiredHeight = InDesiredHeight; BroadcastStyleUpdate(); return *this; }
-};
-
 /**
  * Represents the appearance of a Waveform Editor Time Ruler
  */
diff --git a/Engine/Plugins/Editor/WaveformEditor/Source/WaveformEditorWidgets/Public/WaveformEditorTransportCoordinator.h b/Engine/Plugins/Editor/WaveformEditor/Source/WaveformEditorWidgets/Public/WaveformEditorTransportCoordinator.h
deleted file mode 100644
diff --git a/Engine/Plugins/Runtime/AudioWidgets/Source/AudioWidgets/Private/AudioWidgetsSlateTypes.cpp b/Engine/Plugins/Runtime/AudioWidgets/Source/AudioWidgets/Private/AudioWidgetsSlateTypes.cpp
--- a/Engine/Plugins/Runtime/AudioWidgets/Source/AudioWidgets/Private/AudioWidgetsSlateTypes.cpp
+++ b/Engine/Plugins/Runtime/AudioWidgets/Source/AudioWidgets/Private/AudioWidgetsSlateTypes.cpp
@@ -8,6 +8,15 @@
 
 #include UE_INLINE_GENERATED_CPP_BY_NAME(AudioWidgetsSlateTypes)
 
+namespace AudioWidgetStylesSharedParams
+{
+	const FLazyName BackgroundBrushName = "WhiteBrush";
+	const FLazyName PlayheadHandleBrushName = "Sequencer.Timeline.VanillaScrubHandleDown";
+	const FLinearColor PlayheadColor = FLinearColor(255.f, 0.1f, 0.2f, 1.f);
+	const float DefaultHeight = 720.f;
+	const float DefaultWidth = 1280.f;
+}
+
 // Audio Text Box Style 
 FAudioTextBoxStyle::FAudioTextBoxStyle()
 	: BackgroundImage(FSlateRoundedBoxBrush(FStyleColors::White, 4.0f, FVector2D(56.0f, 28.0f)))
@@ -84,9 +93,9 @@ FSampledSequenceViewerStyle::FSampledSequenceViewerStyle()
 	, ZeroCrossingLineThickness(1.f)
 	, SampleMarkersSize(2.5f)
 	, SequenceBackgroundColor(FLinearColor(0.02f, 0.02f, 0.02f, 1.f))
-	, BackgroundBrush(*FAppStyle::GetBrush(TEXT("WhiteBrush")))
-	, DesiredWidth(1280.f)
-	, DesiredHeight(720.f)
+	, BackgroundBrush(*FAppStyle::GetBrush(AudioWidgetStylesSharedParams::BackgroundBrushName))
+	, DesiredWidth(AudioWidgetStylesSharedParams::DefaultWidth)
+	, DesiredHeight(AudioWidgetStylesSharedParams::DefaultHeight)
 {
 }
 
@@ -101,3 +110,18 @@ void FSampledSequenceViewerStyle::GetResources(TArray< const FSlateBrush* >& Out
 	OutBrushes.Add(&BackgroundBrush);
 }
 
+const FName FPlayheadOverlayStyle::TypeName("FPlayheadOverlayStyle");
+
+FPlayheadOverlayStyle::FPlayheadOverlayStyle()
+	: PlayheadColor(AudioWidgetStylesSharedParams::PlayheadColor)
+	, PlayheadWidth(1.0f)
+	, DesiredWidth(AudioWidgetStylesSharedParams::DefaultWidth)
+	, DesiredHeight(AudioWidgetStylesSharedParams::DefaultHeight)
+{
+}
+
+const FPlayheadOverlayStyle& FPlayheadOverlayStyle::GetDefault()
+{
+	static FPlayheadOverlayStyle Default;
+	return Default;
+}
\ No newline at end of file
diff --git a/Engine/Plugins/Runtime/AudioWidgets/Source/AudioWidgets/Private/AudioWidgetsUtils.cpp b/Engine/Plugins/Runtime/AudioWidgets/Source/AudioWidgets/Private/AudioWidgetsUtils.cpp
new file mode 100644
--- /dev/null
+++ b/Engine/Plugins/Runtime/AudioWidgets/Source/AudioWidgets/Private/AudioWidgetsUtils.cpp
@@ -0,0 +1,71 @@
+// Copyright Epic Games, Inc. All Rights Reserved.
+
+#include "AudioWidgetsUtils.h"
+
+FReply AudioWidgetsUtils::RouteMouseInput(WidgetMouseInputFunction InputFunction, const FPointerEvent& MouseEvent, TArrayView<TSharedPtr<SWidget>> WidgetsStack, const bool bReverse)
+{
+	FReply TransformationInteraction = FReply::Unhandled();
+
+	const int32 StartIndex = bReverse ? WidgetsStack.Num() - 1 : 0;
+	
+	auto ExitCondition = [&bReverse, &WidgetsStack](const int32 i)
+	{
+		return bReverse ? i >= 0 : i < WidgetsStack.Num();
+	};
+
+	auto Increment = [&bReverse](int32& i)
+	{
+		i = bReverse ? i - 1 : i + 1;
+	};
+
+	for (int32 i = StartIndex; ExitCondition(i); Increment(i))
+	{
+		const TSharedPtr<SWidget>& LayerWidget = WidgetsStack[i];
+
+		if (LayerWidget)
+		{
+			TransformationInteraction = (LayerWidget.Get()->*InputFunction)(LayerWidget->GetTickSpaceGeometry(), MouseEvent);
+
+			if (TransformationInteraction.IsEventHandled())
+			{
+				return TransformationInteraction;
+			}
+		}
+	}
+
+	return TransformationInteraction;
+}
+
+FCursorReply AudioWidgetsUtils::RouteCursorQuery(const FPointerEvent& CursorEvent, TArrayView<const TSharedPtr<SWidget>> WidgetsStack, const bool bReverse)
+{
+	FCursorReply CursorReply = FCursorReply::Unhandled();
+
+	const int32 StartIndex = bReverse ? WidgetsStack.Num() - 1 : 0;
+
+	auto ExitCondition = [&bReverse, &WidgetsStack](const int32 i)
+	{
+		return bReverse ? i >= 0 : i < WidgetsStack.Num();
+	};
+
+	auto Increment = [&bReverse](int32& i)
+	{
+		i = bReverse ? i - 1 : i + 1;
+	};
+
+	for (int32 i = StartIndex; ExitCondition(i); i = bReverse ? i - 1 : i + 1)
+	{
+		const TSharedPtr<SWidget>& LayerWidget = WidgetsStack[i];
+
+		if (LayerWidget)
+		{
+			CursorReply = LayerWidget->OnCursorQuery(LayerWidget->GetTickSpaceGeometry(), CursorEvent);
+
+			if (CursorReply.IsEventHandled())
+			{
+				return CursorReply;
+			}
+		}
+	}
+
+	return CursorReply;
+}
\ No newline at end of file
diff --git a/Engine/Plugins/Runtime/AudioWidgets/Source/AudioWidgets/Private/SPlayheadOverlay.cpp b/Engine/Plugins/Runtime/AudioWidgets/Source/AudioWidgets/Private/SPlayheadOverlay.cpp
new file mode 100644
--- /dev/null
+++ b/Engine/Plugins/Runtime/AudioWidgets/Source/AudioWidgets/Private/SPlayheadOverlay.cpp
@@ -0,0 +1,67 @@
+// Copyright Epic Games, Inc. All Rights Reserved.
+
+#include "SPlayheadOverlay.h"
+
+void SPlayheadOverlay::Construct(const FArguments& InArgs)
+{
+	Style = InArgs._Style;
+
+	check(Style);
+	PlayheadColor = Style->PlayheadColor;
+	PlayheadWidth = Style->PlayheadWidth;
+	DesiredWidth = Style->DesiredWidth;
+	DesiredHeight = Style->DesiredHeight;
+}
+
+void SPlayheadOverlay::SetPlayheadPosition(const float InNewPosition)
+{
+	PlayheadPosition = InNewPosition;
+}
+
+void SPlayheadOverlay::OnStyleUpdated(const FNotifyingAudioWidgetStyle& UpdatedStyle)
+{
+	check(Style);
+
+	if (&UpdatedStyle != Style)
+	{
+		return;
+	}
+
+	PlayheadColor = Style->PlayheadColor;
+	PlayheadWidth = Style->PlayheadWidth;
+	DesiredWidth = Style->DesiredWidth;
+	DesiredHeight = Style->DesiredHeight;
+}
+
+int32 SPlayheadOverlay::OnPaint(const FPaintArgs& Args, const FGeometry& AllottedGeometry, const FSlateRect& MyCullingRect, FSlateWindowElementList& OutDrawElements, int32 LayerId, const FWidgetStyle& InWidgetStyle, bool bParentEnabled) const
+{
+	LayerId = DrawPlayhead(AllottedGeometry, OutDrawElements, LayerId);
+
+	return LayerId;
+}
+
+int32 SPlayheadOverlay::DrawPlayhead(const FGeometry& AllottedGeometry, FSlateWindowElementList& OutDrawElements, int32 LayerId) const
+{
+	const float PlayHeadX = PlayheadPosition;
+
+	TArray<FVector2D> LinePoints{ FVector2D(PlayHeadX, 0.0f), FVector2D(PlayHeadX, AllottedGeometry.Size.Y) };
+	
+
+	FSlateDrawElement::MakeLines(
+		OutDrawElements,
+		LayerId,
+		AllottedGeometry.ToPaintGeometry(),
+		LinePoints,
+		ESlateDrawEffect::None,
+		PlayheadColor.GetSpecifiedColor(),
+		true,
+		PlayheadWidth
+	);
+
+	return ++LayerId;
+}
+
+FVector2D SPlayheadOverlay::ComputeDesiredSize(float) const
+{
+	return FVector2D(DesiredWidth, DesiredHeight);
+}
\ No newline at end of file
diff --git a/Engine/Plugins/Runtime/AudioWidgets/Source/AudioWidgets/Private/SSampledSequenceViewer.cpp b/Engine/Plugins/Runtime/AudioWidgets/Source/AudioWidgets/Private/SSampledSequenceViewer.cpp
--- a/Engine/Plugins/Runtime/AudioWidgets/Source/AudioWidgets/Private/SSampledSequenceViewer.cpp
+++ b/Engine/Plugins/Runtime/AudioWidgets/Source/AudioWidgets/Private/SSampledSequenceViewer.cpp
@@ -240,12 +240,11 @@ void SSampledSequenceViewer::UpdateGridMetrics()
 	GridMetrics = GridService->GetGridMetrics();
 }
 
-void SSampledSequenceViewer::OnStyleUpdated(const FNotifyingAudioWidgetStyle* UpdatedStyle)
+void SSampledSequenceViewer::OnStyleUpdated(const FNotifyingAudioWidgetStyle& UpdatedStyle)
 {
-	check(UpdatedStyle);
 	check(Style);
 
-	if (UpdatedStyle != Style)
+	if (&UpdatedStyle != Style)
 	{
 		return;
 	}
diff --git a/Engine/Plugins/Runtime/AudioWidgets/Source/AudioWidgets/Private/SamplesSequenceTransportCoordinator.cpp b/Engine/Plugins/Runtime/AudioWidgets/Source/AudioWidgets/Private/SamplesSequenceTransportCoordinator.cpp
new file mode 100644
--- /dev/null
+++ b/Engine/Plugins/Runtime/AudioWidgets/Source/AudioWidgets/Private/SamplesSequenceTransportCoordinator.cpp
@@ -0,0 +1,238 @@
+// Copyright Epic Games, Inc. All Rights Reserved.
+
+#include "SamplesSequenceTransportCoordinator.h"
+
+/** Zoom Formula
+
+E.g: Coordinator is receiving a playback ratio (PBR) in the range 0 to 1
+To position the focus point accordingly, this would the formula to follow:
+
+FP = (PBR - LB) / ZR
+
+FP = Focus Point
+ZR = Zoom Ratio
+LB = Lower Displayed Waveform Bound
+from here, you can rotate as needed. 
+e.g to find PBR
+PBR = (FP * ZR) + LB
+
+
+
+This formula can be used to position any element (e.g. a marker) given in a 0 to 1 position range where
+0 is the beginning time of the waveform and 1 the end time
+
+Zoom is defined as a ratio = displayed length / total length
+
+*/
+
+
+
+void FSamplesSequenceTransportCoordinator::UpdatePlaybackRange(const TRange<float>& NewRange)
+{
+	ProgressRange = NewRange;
+}
+
+
+
+const bool FSamplesSequenceTransportCoordinator::IsScrubbing() const
+{
+	return bIsScrubbing;
+}
+
+
+
+float FSamplesSequenceTransportCoordinator::ConvertAbsoluteRatioToZoomed(const float InAbsoluteRatio) const
+{
+	return (InAbsoluteRatio - DisplayRange.GetLowerBoundValue()) / ZoomRatio;
+}
+
+
+
+float FSamplesSequenceTransportCoordinator::ConvertZoomedRatioToAbsolute(const float InZoomedRatio) const
+{
+	return InZoomedRatio * ZoomRatio + DisplayRange.GetLowerBoundValue();
+}
+
+const TRange<float> FSamplesSequenceTransportCoordinator::GetDisplayRange() const
+{
+	return DisplayRange;
+}
+
+void FSamplesSequenceTransportCoordinator::MoveFocusPoint(const float InFocusPoint)
+{
+	FocusPoint = InFocusPoint;
+}
+
+void FSamplesSequenceTransportCoordinator::UpdateZoomRatioAndDisplayRange(const float NewZoomRatio)
+{
+	const float ClampedZoomRatio = FMath::Clamp(NewZoomRatio, UE_KINDA_SMALL_NUMBER, 1.f);
+	const float ZoomRatioDelta = ZoomRatio - ClampedZoomRatio;
+
+ 	const float DeltaOutsideBoundaries = DisplayRange.GetUpperBoundValue() - (DisplayRange.GetUpperBoundValue() - ZoomRatioDelta);
+
+	const float DeltaStepL = DeltaOutsideBoundaries * FocusPoint;
+	const float DeltaStepR = DeltaOutsideBoundaries * (1.f - FocusPoint);
+
+	float MinDisplayRangeValue = DeltaStepL + DisplayRange.GetLowerBoundValue();
+	float MaxDisplayRangeValue = DisplayRange.GetUpperBoundValue() - DeltaStepR;
+
+	const float CurrentPlaybackPosition = (ZoomRatio * FocusPoint) + DisplayRange.GetLowerBoundValue();
+
+	//if Min < 0 or Max > 1 we can't zoom out further on that side
+
+	//hence we shift the zooming entirely on the opposite side
+	if (MinDisplayRangeValue < 0.f)
+	{
+		const float ExceedingDelta = FMath::Abs(MinDisplayRangeValue);
+		MaxDisplayRangeValue = MaxDisplayRangeValue + ExceedingDelta;
+		MinDisplayRangeValue = 0.f;
+	}
+
+	else if (MaxDisplayRangeValue > 1.f)
+	{
+		const float ExceedingDelta = MaxDisplayRangeValue - 1.f;
+
+		MinDisplayRangeValue = MinDisplayRangeValue - ExceedingDelta;
+		MaxDisplayRangeValue = 1.f;
+	}
+
+	const float NewPlayheadPosition = (CurrentPlaybackPosition - MinDisplayRangeValue) / ClampedZoomRatio;
+
+	ZoomRatio = ClampedZoomRatio;
+	UpdateDisplayRange(MinDisplayRangeValue, MaxDisplayRangeValue);
+	MoveFocusPoint(NewPlayheadPosition);
+}
+
+void FSamplesSequenceTransportCoordinator::Stop()
+{
+	SetProgressRatio(0.f);
+}
+
+const float FSamplesSequenceTransportCoordinator::GetFocusPoint() const
+{
+	return FocusPoint;
+}
+
+void FSamplesSequenceTransportCoordinator::ScrubFocusPoint(const float InTargetFocusPoint, const bool bIsMoving)
+{
+	const float LowerPlaybackBoundPosition = (ProgressRange.GetLowerBoundValue() - DisplayRange.GetLowerBoundValue()) / ZoomRatio;
+	const float UpperPlaybackBoundPosition = (ProgressRange.GetUpperBoundValue() - DisplayRange.GetLowerBoundValue()) / ZoomRatio;
+
+	float ClampedTargetPosition = FMath::Clamp(InTargetFocusPoint, LowerPlaybackBoundPosition, UpperPlaybackBoundPosition);
+	float NewPosition = FMath::Clamp(ClampedTargetPosition, 0.f, 1.f);
+
+	if (ClampedTargetPosition > FocusPointLockPosition)
+	{
+		const float BoundsDelta = (ClampedTargetPosition - FocusPointLockPosition) * ZoomRatio;
+		const float NewDisplayUpBound = FMath::Clamp(DisplayRange.GetUpperBoundValue() + BoundsDelta, 0.f, 1.f);
+
+		if (NewDisplayUpBound < 1.f)
+		{
+			const float NewDisplayLowBound = FMath::Clamp(DisplayRange.GetLowerBoundValue() + BoundsDelta, 0.f, 1.f);
+			UpdateDisplayRange(NewDisplayLowBound, NewDisplayUpBound);
+			NewPosition = FocusPointLockPosition;
+		}
+	}
+	if (ClampedTargetPosition < 0.f)
+	{
+		const float BoundsDelta = ClampedTargetPosition * ZoomRatio;
+		const float NewDisplayLowBound = FMath::Clamp(DisplayRange.GetLowerBoundValue() + BoundsDelta, 0.f, 1.f);
+
+		if (NewDisplayLowBound > 0.f)
+		{
+
+			const float NewDisplayUpBound = FMath::Clamp(DisplayRange.GetUpperBoundValue() + BoundsDelta, 0.f, 1.f);
+
+			UpdateDisplayRange(NewDisplayLowBound, NewDisplayUpBound);
+
+			NewPosition = 0.f;
+		}
+	}
+
+	MoveFocusPoint(NewPosition);
+
+	if (OnFocusPointScrubUpdate.IsBound())
+	{
+		OnFocusPointScrubUpdate.Broadcast(GetPlayBackRatioFromFocusPoint(NewPosition), bIsMoving);
+	}
+
+	bIsScrubbing = bIsMoving;
+}
+
+
+
+void FSamplesSequenceTransportCoordinator::SetProgressRatio(const float NewRatio)
+{	 
+	if (bIsScrubbing)
+	{
+		return;
+	}
+
+	check(NewRatio >= 0.f && NewRatio <= 1.f);
+
+	const float PlaybackRangeScale = ProgressRange.Size<float>();
+	const float ScaledPlaybackRatio = NewRatio * PlaybackRangeScale + ProgressRange.GetLowerBoundValue();
+
+	const float NewPlayheadPosition = (ScaledPlaybackRatio - DisplayRange.GetLowerBoundValue()) / ZoomRatio;
+
+	MoveFocusPoint(NewPlayheadPosition);
+
+	if (FocusPoint < 0.f)
+	{
+		//if sound is looping and we are zoomed, playhead position will be negative when new loop starts
+		//hence we reset the view to go back at the beginning of the file
+
+		MoveFocusPoint(0.f);
+		UpdateDisplayRange(FocusPoint, ZoomRatio);
+	}
+
+	else if (FocusPoint >= FocusPointLockPosition)
+	{
+
+		const float NewDisplayUpBound = FMath::Clamp(ScaledPlaybackRatio + (ZoomRatio * (1.f - FocusPointLockPosition)), 0.f, 1.f);
+		float NewDisplayLowBound = NewDisplayUpBound < 1.f ? ScaledPlaybackRatio - (ZoomRatio * FocusPointLockPosition) : 1.f - ZoomRatio;
+
+		NewDisplayLowBound = FMath::Clamp(NewDisplayLowBound, 0.f, 1.f);
+
+		if (NewDisplayUpBound < 1.f)
+		{
+			//we only lock the playhead if the display upper bound is before the end of the waveform
+			//otherwise we let the playhead go to the end
+
+			MoveFocusPoint(FocusPointLockPosition);
+		}
+
+		UpdateDisplayRange(NewDisplayLowBound, NewDisplayUpBound);
+	}
+}
+
+void FSamplesSequenceTransportCoordinator::SetZoomRatio(const float NewLevel)
+{
+	UpdateZoomRatioAndDisplayRange(NewLevel);
+}
+
+void FSamplesSequenceTransportCoordinator::UpdateDisplayRange(const float MinValue, const float MaxValue)
+{
+	check(MinValue < MaxValue);
+
+	DisplayRange.SetLowerBoundValue(MinValue);
+	DisplayRange.SetUpperBoundValue(MaxValue);
+
+	if (OnDisplayRangeUpdated.IsBound())
+	{
+		OnDisplayRangeUpdated.Broadcast(DisplayRange);
+	}
+}
+
+bool FSamplesSequenceTransportCoordinator::IsRatioWithinDisplayRange(const float Ratio) const
+{
+	return DisplayRange.Contains(Ratio);
+}
+
+float FSamplesSequenceTransportCoordinator::GetPlayBackRatioFromFocusPoint(const float InFocusPoint) const
+{
+	check(InFocusPoint >= 0.f && InFocusPoint <= 1.f);
+
+	const float PlayBackRatio = ((InFocusPoint * ZoomRatio) + DisplayRange.GetLowerBoundValue() - ProgressRange.GetLowerBoundValue()) / ProgressRange.Size<float>();
+	return PlayBackRatio;
+}
\ No newline at end of file
diff --git a/Engine/Plugins/Runtime/AudioWidgets/Source/AudioWidgets/Public/AudioWidgetsSlateTypes.h b/Engine/Plugins/Runtime/AudioWidgets/Source/AudioWidgets/Public/AudioWidgetsSlateTypes.h
--- a/Engine/Plugins/Runtime/AudioWidgets/Source/AudioWidgets/Public/AudioWidgetsSlateTypes.h
+++ b/Engine/Plugins/Runtime/AudioWidgets/Source/AudioWidgets/Public/AudioWidgetsSlateTypes.h
@@ -7,7 +7,7 @@
 
 #include "AudioWidgetsSlateTypes.generated.h"
 
-DECLARE_MULTICAST_DELEGATE_OneParam(FOnAudioWidgetStyleUpdated, const FNotifyingAudioWidgetStyle* /*Updated Widget Style*/);
+DECLARE_MULTICAST_DELEGATE_OneParam(FOnAudioWidgetStyleUpdated, const FNotifyingAudioWidgetStyle& /*Updated Widget Style*/);
 
 USTRUCT()
 struct FNotifyingAudioWidgetStyle : public FSlateWidgetStyle
@@ -17,7 +17,7 @@ struct FNotifyingAudioWidgetStyle : public FSlateWidgetStyle
 	FNotifyingAudioWidgetStyle() = default;
 	virtual ~FNotifyingAudioWidgetStyle() = default;
 
-	virtual void BroadcastStyleUpdate() const { OnStyleUpdated.Broadcast(this); }
+	virtual void BroadcastStyleUpdate() const { OnStyleUpdated.Broadcast(*this); }
 
 	FOnAudioWidgetStyleUpdated OnStyleUpdated;
 };
@@ -247,6 +247,38 @@ struct AUDIOWIDGETS_API FSampledSequenceViewerStyle : public FNotifyingAudioWidg
 	FSampledSequenceViewerStyle& SetDesiredHeight(const float InDesiredHeight) { DesiredHeight = InDesiredHeight; BroadcastStyleUpdate(); return *this; }
 };
 
+/**
+ * Represents the appearance of a Waveform Viewer Overlay style
+ */
+USTRUCT(BlueprintType)
+struct AUDIOWIDGETS_API FPlayheadOverlayStyle : public FNotifyingAudioWidgetStyle
+{
+	GENERATED_USTRUCT_BODY()
+
+	FPlayheadOverlayStyle();
+
+	static const FPlayheadOverlayStyle& GetDefault();
+	virtual const FName GetTypeName() const override { return TypeName; };
+
+	static const FName TypeName;
+
+	UPROPERTY(EditAnywhere, BlueprintReadWrite, Category = Appearance)
+	FSlateColor PlayheadColor;
+	FPlayheadOverlayStyle& SetPlayheadColor(const FSlateColor InPlayheadColor) { PlayheadColor = InPlayheadColor; BroadcastStyleUpdate(); return *this; }
+
+	UPROPERTY(EditAnywhere, BlueprintReadWrite, Category = Appearance)
+	float PlayheadWidth;
+	FPlayheadOverlayStyle& SetPlayheadWidth(const float InPlayheadWidth) { PlayheadWidth = InPlayheadWidth; BroadcastStyleUpdate(); return *this; }
+
+	UPROPERTY(EditAnywhere, BlueprintReadWrite, Category = Appearance)
+	float DesiredWidth;
+	FPlayheadOverlayStyle& SetDesiredWidth(const float InDesiredWidth) { DesiredWidth = InDesiredWidth; BroadcastStyleUpdate(); return *this; }
+
+	UPROPERTY(EditAnywhere, BlueprintReadWrite, Category = Appearance)
+	float DesiredHeight;
+	FPlayheadOverlayStyle& SetDesiredHeight(const float InDesiredHeight) { DesiredHeight = InDesiredHeight; BroadcastStyleUpdate(); return *this; }
+};
+
 #if UE_ENABLE_INCLUDE_ORDER_DEPRECATED_IN_5_2
 #include "Styling/SlateStyle.h"
 #include "Styling/SlateStyleMacros.h"
diff --git a/Engine/Plugins/Runtime/AudioWidgets/Source/AudioWidgets/Public/AudioWidgetsUtils.h b/Engine/Plugins/Runtime/AudioWidgets/Source/AudioWidgets/Public/AudioWidgetsUtils.h
new file mode 100644
--- /dev/null
+++ b/Engine/Plugins/Runtime/AudioWidgets/Source/AudioWidgets/Public/AudioWidgetsUtils.h
@@ -0,0 +1,23 @@
+// Copyright Epic Games, Inc. All Rights Reserved.
+
+#pragma once
+
+#include "Widgets/SWidget.h"
+
+namespace AudioWidgetsUtils
+{
+	typedef FReply(SWidget::* WidgetMouseInputFunction)(const FGeometry&, const FPointerEvent&);
+
+	/**
+	*	Propagates mouse input through an array of widget until that input is handled 
+	*	or all the elements have been iterated
+	* 
+	*/
+	FReply AUDIOWIDGETS_API RouteMouseInput(WidgetMouseInputFunction InputFunction, const FPointerEvent& MouseEvent, TArrayView<TSharedPtr<SWidget>> WidgetsStack, const bool bReverse = false);
+
+	/**
+	*	Propagates a cursor query through an array of widget until that input is handled
+	*	or all the elements have been iterated
+	*/
+	FCursorReply AUDIOWIDGETS_API RouteCursorQuery(const FPointerEvent& CursorEvent, TArrayView<const TSharedPtr<SWidget>> WidgetsStack, const bool bReverse = false);
+}
\ No newline at end of file
diff --git a/Engine/Plugins/Runtime/AudioWidgets/Source/AudioWidgets/Public/ISamplesSequenceTransportCoordinator.h b/Engine/Plugins/Runtime/AudioWidgets/Source/AudioWidgets/Public/ISamplesSequenceTransportCoordinator.h
new file mode 100644
--- /dev/null
+++ b/Engine/Plugins/Runtime/AudioWidgets/Source/AudioWidgets/Public/ISamplesSequenceTransportCoordinator.h
@@ -0,0 +1,23 @@
+// Copyright Epic Games, Inc. All Rights Reserved.
+
+#pragma once
+
+#include "Math/Range.h"
+
+class AUDIOWIDGETS_API ISamplesSequenceTransportCoordinator
+{
+public:
+	virtual ~ISamplesSequenceTransportCoordinator() = default;
+
+	virtual const TRange<float> GetDisplayRange() const = 0;
+	virtual const float GetFocusPoint() const = 0;
+	virtual void ScrubFocusPoint(const float InTargetFocusPoint, const bool bIsMoving) = 0;
+	virtual const bool IsScrubbing() const = 0;
+
+	virtual void SetProgressRatio(const float NewRatio) = 0;
+	virtual void SetZoomRatio(const float NewRatio) = 0;
+
+	virtual float ConvertAbsoluteRatioToZoomed(const float InAbsoluteRatio) const = 0;
+	virtual float ConvertZoomedRatioToAbsolute(const float InZoomedRatio) const = 0;
+
+};
\ No newline at end of file
diff --git a/Engine/Plugins/Runtime/AudioWidgets/Source/AudioWidgets/Public/SPlayheadOverlay.h b/Engine/Plugins/Runtime/AudioWidgets/Source/AudioWidgets/Public/SPlayheadOverlay.h
new file mode 100644
--- /dev/null
+++ b/Engine/Plugins/Runtime/AudioWidgets/Source/AudioWidgets/Public/SPlayheadOverlay.h
@@ -0,0 +1,39 @@
+// Copyright Epic Games, Inc. All Rights Reserved.
+
+#pragma once
+
+#include "AudioWidgetsSlateTypes.h"
+#include "Styling/ISlateStyle.h"
+#include "Styling/SlateWidgetStyleAsset.h"
+#include "Widgets/SLeafWidget.h"
+
+class AUDIOWIDGETS_API SPlayheadOverlay : public SLeafWidget
+{
+public:
+	SLATE_BEGIN_ARGS(SPlayheadOverlay)
+	{
+	}
+
+	SLATE_STYLE_ARGUMENT(FPlayheadOverlayStyle, Style)
+
+	SLATE_END_ARGS()
+
+	void Construct(const FArguments& InArgs);
+	void SetPlayheadPosition(const float InNewPosition);
+	void OnStyleUpdated(const FNotifyingAudioWidgetStyle& UpdatedStyle);
+
+private:
+	virtual int32 OnPaint(const FPaintArgs& Args, const FGeometry& AllottedGeometry, const FSlateRect& MyCullingRect, FSlateWindowElementList& OutDrawElements, int32 LayerId, const FWidgetStyle& InWidgetStyle, bool bParentEnabled) const override;
+	virtual FVector2D ComputeDesiredSize(float) const override;
+
+	int32 DrawPlayhead(const FGeometry& AllottedGeometry, FSlateWindowElementList& OutDrawElements, int32 LayerId) const;
+
+	const FPlayheadOverlayStyle* Style = nullptr;
+
+	FSlateColor PlayheadColor = FLinearColor(255.f, 0.1f, 0.2f, 1.f);
+	float PlayheadWidth = 1.0;
+	float DesiredWidth = 0.f;
+	float DesiredHeight = 0.f;
+
+	float PlayheadPosition = 0.f;
+};
\ No newline at end of file
diff --git a/Engine/Plugins/Runtime/AudioWidgets/Source/AudioWidgets/Public/SSampledSequenceViewer.h b/Engine/Plugins/Runtime/AudioWidgets/Source/AudioWidgets/Public/SSampledSequenceViewer.h
--- a/Engine/Plugins/Runtime/AudioWidgets/Source/AudioWidgets/Public/SSampledSequenceViewer.h
+++ b/Engine/Plugins/Runtime/AudioWidgets/Source/AudioWidgets/Public/SSampledSequenceViewer.h
@@ -31,7 +31,9 @@ public:
 	void Construct(const FArguments& InArgs, TArrayView<const float> InSampleData, const uint8 InNumChannels, TSharedRef<ISampledSequenceGridService> InGridService);
 	void UpdateView(TArrayView<const float> InSampleData, const uint8 InNumChannels);
 	
-	void OnStyleUpdated(const FNotifyingAudioWidgetStyle* UpdatedStyle);
+	void OnStyleUpdated(const FNotifyingAudioWidgetStyle& UpdatedStyle);
+
+	void UpdateGridMetrics();
 
 private:
 	virtual int32 OnPaint(const FPaintArgs& Args, const FGeometry& AllottedGeometry, const FSlateRect& MyCullingRect, FSlateWindowElementList& OutDrawElements, int32 LayerId, const FWidgetStyle& InWidgetStyle, bool bParentEnabled) const override;
@@ -40,8 +42,6 @@ private:
 
 	void DrawGridLines(const FGeometry& AllottedGeometry, FSlateWindowElementList& OutDrawElements, int32& LayerId) const;
 
-	void UpdateGridMetrics();
-
 	bool bForceRedraw = false;
 
 	TArrayView<const float> SampleData;
diff --git a/Engine/Plugins/Runtime/AudioWidgets/Source/AudioWidgets/Public/SamplesSequenceTransportCoordinator.h b/Engine/Plugins/Runtime/AudioWidgets/Source/AudioWidgets/Public/SamplesSequenceTransportCoordinator.h
new file mode 100644
--- /dev/null
+++ b/Engine/Plugins/Runtime/AudioWidgets/Source/AudioWidgets/Public/SamplesSequenceTransportCoordinator.h
@@ -0,0 +1,58 @@
+// Copyright Epic Games, Inc. All Rights Reserved.
+
+#pragma once
+
+#include "Delegates/Delegate.h"
+#include "ISamplesSequenceTransportCoordinator.h"
+
+DECLARE_MULTICAST_DELEGATE_OneParam(FOnDisplayRangeUpdated, const TRange<float> /* New Display Range */);
+DECLARE_MULTICAST_DELEGATE_TwoParams(FOnFocusPointScrubUpdate, const float /* Focused Playback Ratio */, const bool /*Playhead is Moving*/);
+
+
+class AUDIOWIDGETS_API FSamplesSequenceTransportCoordinator : public ISamplesSequenceTransportCoordinator
+{
+public:
+	FSamplesSequenceTransportCoordinator() = default;
+	virtual ~FSamplesSequenceTransportCoordinator() = default;
+
+	/** Called when the playhead is scrubbed */
+	FOnFocusPointScrubUpdate OnFocusPointScrubUpdate;
+
+	/** Called when the display range is updated */
+	FOnDisplayRangeUpdated OnDisplayRangeUpdated;
+
+	/** ISamplesSequenceTransportCoordinator interface */
+	const TRange<float> GetDisplayRange() const;
+	const float GetFocusPoint() const override;
+	void ScrubFocusPoint(const float InTargetFocusPoint, const bool bIsMoving) override;
+	const bool IsScrubbing() const;
+	void SetProgressRatio(const float NewRatio) override;
+	void SetZoomRatio(const float NewRatio) override;
+	float ConvertAbsoluteRatioToZoomed(const float InAbsoluteRatio) const override;
+	float ConvertZoomedRatioToAbsolute(const float InZoomedRatio) const override;
+
+	void UpdatePlaybackRange(const TRange<float>& NewRange);
+	void Stop();
+
+private:
+	FORCEINLINE void MoveFocusPoint(const float InFocusPoint);
+	void UpdateZoomRatioAndDisplayRange(const float NewZoomRatio);
+
+	void UpdateDisplayRange(const float MinValue, const float MaxValue);
+	bool IsRatioWithinDisplayRange(const float Ratio) const;
+
+	float GetPlayBackRatioFromFocusPoint(const float InFocusPoint) const;
+	
+	float CurrentPlaybackRatio = 0.f;
+	float FocusPointLockPosition = 0.95f;
+	float FocusPoint = 0.f;
+	float ZoomRatio = 1.f;
+	
+	/* The currently displayed render data range */
+	TRange<float> DisplayRange = TRange<float>::Inclusive(0.f, 1.f);
+
+	/* Progress range to scale the incoming progress ratio with*/
+	TRange<float> ProgressRange = TRange<float>::Inclusive(0.f, 1.f);
+	
+	bool bIsScrubbing = false;
+};
\ No newline at end of file
