[WaveformEditor] Move time ruler into audiowidgets

diff --git a/Engine/Plugins/Editor/WaveformEditor/Source/WaveformEditorWidgets/Private/SWaveformEditorTimeRuler.cpp b/Engine/Plugins/Editor/WaveformEditor/Source/WaveformEditorWidgets/Private/SWaveformEditorTimeRuler.cpp
deleted file mode 100644
diff --git a/Engine/Plugins/Editor/WaveformEditor/Source/WaveformEditorWidgets/Private/SWaveformEditorTimeRuler.h b/Engine/Plugins/Editor/WaveformEditor/Source/WaveformEditorWidgets/Private/SWaveformEditorTimeRuler.h
deleted file mode 100644
diff --git a/Engine/Plugins/Editor/WaveformEditor/Source/WaveformEditorWidgets/Private/SWaveformPanel.cpp b/Engine/Plugins/Editor/WaveformEditor/Source/WaveformEditorWidgets/Private/SWaveformPanel.cpp
--- a/Engine/Plugins/Editor/WaveformEditor/Source/WaveformEditorWidgets/Private/SWaveformPanel.cpp
+++ b/Engine/Plugins/Editor/WaveformEditor/Source/WaveformEditorWidgets/Private/SWaveformPanel.cpp
@@ -2,13 +2,13 @@
 
 #include "SWaveformPanel.h"
 
-#include "SPlayheadOverlay.h"
+#include "SampledSequenceDisplayUnit.h"
+#include "SFixedSampledSequenceRuler.h"
 #include "SFixedSampledSequenceViewer.h"
-#include "SWaveformEditorTimeRuler.h"
+#include "SparseSampledSequenceTransportCoordinator.h"
+#include "SPlayheadOverlay.h"
 #include "SWaveformTransformationsOverlay.h"
 #include "SWaveformViewerOverlay.h"
-#include "SparseSampledSequenceTransportCoordinator.h"
-#include "WaveformEditorDisplayUnit.h"
 #include "WaveformEditorGridData.h"
 #include "WaveformEditorRenderData.h"
 #include "WaveformEditorStyle.h"
@@ -19,15 +19,17 @@
 
 void SWaveformPanel::Construct(const FArguments& InArgs, TSharedRef<FWaveformEditorRenderData> InRenderData, TSharedRef<FSparseSampledSequenceTransportCoordinator> InTransportCoordinator, TSharedRef<FWaveformEditorZoomController> InZoomManager, TSharedPtr<SWaveformTransformationsOverlay> InWaveformTransformationsOverlay)
 {
-	DisplayUnit = EWaveformEditorDisplayUnit::Seconds;
 
+	DisplayUnit = ESampledSequenceDisplayUnit::Seconds;
 	WaveformEditorStyle = &FWaveformEditorStyle::Get();
+
 	check(WaveformEditorStyle);
 
 	RenderData = InRenderData;
+
 	GenerateFloatRenderData();
-	InRenderData->OnRenderDataUpdated.AddSP(this, &SWaveformPanel::OnRenderDataUpdated);
 
+	InRenderData->OnRenderDataUpdated.AddSP(this, &SWaveformPanel::OnRenderDataUpdated);
 	TransportCoordinator = InTransportCoordinator;
 	TransportCoordinator->OnDisplayRangeUpdated.AddSP(this, &SWaveformPanel::OnDisplayRangeUpdated);
 
@@ -39,17 +41,14 @@ void SWaveformPanel::Construct(const FArguments& InArgs, TSharedRef<FWaveformEdi
 	{
 		WaveformTransformationsOverlay = InWaveformTransformationsOverlay;
 	}
-	
-	SetupPlayheadOverlay();
 
+	SetupPlayheadOverlay();
 	SetUpWaveformViewerOverlay(InZoomManager);
 	SetUpTimeRuler(GridData.ToSharedRef());
-	CreateLayout();
 	
+	CreateLayout();
 }
 
-
-
 void SWaveformPanel::CreateLayout()
 {
 	check(TimeRuler);
@@ -57,6 +56,7 @@ void SWaveformPanel::CreateLayout()
 	check(WaveformViewerOverlay);
 
 	TSharedPtr<SOverlay> WaveformView = SNew(SOverlay);
+
 	WaveformView->AddSlot()
 	[
 		WaveformViewer.ToSharedRef()
@@ -80,7 +80,6 @@ void SWaveformPanel::CreateLayout()
 		WaveformViewerOverlay.ToSharedRef()
 	];
 
-
 	ChildSlot
 	[
 		SNew(SVerticalBox)
@@ -97,13 +96,18 @@ void SWaveformPanel::CreateLayout()
 
 void SWaveformPanel::SetUpTimeRuler(TSharedRef<FWaveformEditorGridData> InGridData)
 {
-	FWaveformEditorTimeRulerStyle* TimeRulerStyle = &WaveformEditorStyle->GetRegisteredWidgetStyle<FWaveformEditorTimeRulerStyle>("WaveformEditorRuler.Style").Get();
+	FFixedSampleSequenceRulerStyle* TimeRulerStyle = &WaveformEditorStyle->GetRegisteredWidgetStyle<FFixedSampleSequenceRulerStyle>("WaveformEditorRuler.Style").Get();
+
 	check(TimeRulerStyle);
 
-	TimeRuler = SNew(SWaveformEditorTimeRuler, InGridData).DisplayUnit(DisplayUnit).Style(TimeRulerStyle);
-	GridData->OnGridMetricsUpdated.AddSP(TimeRuler.ToSharedRef(), &SWaveformEditorTimeRuler::UpdateGridMetrics);
-	TimeRulerStyle->OnStyleUpdated.AddSP(TimeRuler.ToSharedRef(), &SWaveformEditorTimeRuler::OnStyleUpdated);
+	TimeRuler = SNew(SFixedSampledSequenceRuler, InGridData).DisplayUnit(DisplayUnit).Style(TimeRulerStyle);
+
+	GridData->OnGridMetricsUpdated.AddSP(TimeRuler.ToSharedRef(), &SFixedSampledSequenceRuler::UpdateGridMetrics);
+
+	TimeRulerStyle->OnStyleUpdated.AddSP(TimeRuler.ToSharedRef(), &SFixedSampledSequenceRuler::OnStyleUpdated);
+
 	TimeRuler->OnTimeUnitMenuSelection.AddSP(this, &SWaveformPanel::UpdateDisplayUnit);
+
 	TimeRuler->SetOnMouseButtonDown(FPointerEventHandler::CreateSP(this, &SWaveformPanel::HandleTimeRulerMouseButtonDown));
 	TimeRuler->SetOnMouseButtonUp(FPointerEventHandler::CreateSP(this, &SWaveformPanel::HandleTimeRulerMouseButtonUp));
 	TimeRuler->SetOnMouseMove(FPointerEventHandler::CreateSP(this, &SWaveformPanel::HandleTimeRulerMouseMove));
@@ -113,7 +117,7 @@ void SWaveformPanel::SetUpWaveformViewerOverlay(TSharedRef<FWaveformEditorZoomCo
 {
 	FSampledSequenceViewerStyle* ViewerStyle = &WaveformEditorStyle->GetRegisteredWidgetStyle<FSampledSequenceViewerStyle>("WaveformViewer.Style").Get();
 	check(ViewerStyle);
-	
+
 	TArray<TSharedPtr<SWidget>> OverlaidWidgets;
 
 	if (WaveformTransformationsOverlay)
@@ -123,7 +127,6 @@ void SWaveformPanel::SetUpWaveformViewerOverlay(TSharedRef<FWaveformEditorZoomCo
 
 	check(PlayheadOverlay)
 	OverlaidWidgets.Add(PlayheadOverlay);
-
 	WaveformViewerOverlay = SNew(SWaveformViewerOverlay, OverlaidWidgets).Style(ViewerStyle);
 	WaveformViewerOverlay->OnNewMouseDelta.BindSP(InZoomManager, &FWaveformEditorZoomController::ZoomByDelta);
 }
@@ -142,10 +145,10 @@ void SWaveformPanel::SetUpWaveformViewer(TSharedRef<FWaveformEditorGridData> InG
 {
 	FSampledSequenceViewerStyle* WaveViewerStyle = &WaveformEditorStyle->GetRegisteredWidgetStyle<FSampledSequenceViewerStyle>("WaveformViewer.Style").Get();
 	check(WaveViewerStyle);
-	
+
 	SampledSequenceDrawingUtils::FSampledSequenceDrawingParams WaveformViewerDrawingParams;
 	WaveformViewerDrawingParams.MaxDisplayedValue = TNumericLimits<int16>::Max();
-
+	
 	WaveformViewer = SNew(SFixedSampledSequenceViewer, MakeArrayView(FloatRenderData.GetData(), FloatRenderData.Num()), InRenderData->GetNumChannels(), InGridData).Style(WaveViewerStyle).SequenceDrawingParams(WaveformViewerDrawingParams);
 	WaveViewerStyle->OnStyleUpdated.AddSP(WaveformViewer.ToSharedRef(), &SFixedSampledSequenceViewer::OnStyleUpdated);
 	GridData->OnGridMetricsUpdated.AddSP(WaveformViewer.ToSharedRef(), &SFixedSampledSequenceViewer::UpdateGridMetrics);
@@ -153,9 +156,9 @@ void SWaveformPanel::SetUpWaveformViewer(TSharedRef<FWaveformEditorGridData> InG
 
 void SWaveformPanel::SetUpGridData(TSharedRef<FWaveformEditorRenderData> InRenderData)
 {
-	FWaveformEditorTimeRulerStyle* RulerStyle = &WaveformEditorStyle->GetRegisteredWidgetStyle<FWaveformEditorTimeRulerStyle>("WaveformEditorRuler.Style").Get();
+	FFixedSampleSequenceRulerStyle* RulerStyle = &WaveformEditorStyle->GetRegisteredWidgetStyle<FFixedSampleSequenceRulerStyle>("WaveformEditorRuler.Style").Get();
 	check(RulerStyle)
-	
+
 	GridData = MakeShared<FWaveformEditorGridData>(InRenderData->GetNumSamples() / InRenderData->GetNumChannels(), InRenderData->GetSampleRate(), RulerStyle->DesiredWidth, &RulerStyle->TicksTextFont);
 }
 
@@ -179,13 +182,12 @@ void SWaveformPanel::OnDisplayRangeUpdated(const TRange<float> NewDisplayRange)
 	const uint32 MinSamplesToDisplay = MinFramesToDisplay * RenderData->GetNumChannels();
 	const uint32 NumOriginalSamples = RenderData->GetSampleData().Num();
 	const uint32 NumOriginalFrames = NumOriginalSamples / RenderData->GetNumChannels();
-
 	const uint32 FirstRenderedSample = FMath::Clamp(FMath::RoundToInt32(NumOriginalFrames * NewDisplayRange.GetLowerBoundValue()), 0, NumOriginalFrames - MinFramesToDisplay) * RenderData->GetNumChannels();	
 	const uint32 NumFramesToRender = FMath::RoundToInt32(NumOriginalFrames * NewDisplayRange.Size<float>());
 	const uint32 NumSamplesToRender = FMath::Clamp(NumFramesToRender * RenderData->GetNumChannels(), MinSamplesToDisplay, NumOriginalSamples - FirstRenderedSample);
-	
+
 	check(NumSamplesToRender % RenderData->GetNumChannels() == 0 && FirstRenderedSample % RenderData->GetNumChannels() == 0);
-	
+
 	if (GridData)
 	{
 		const uint32 FirstRenderedFrame = FirstRenderedSample / RenderData->GetNumChannels();
@@ -203,8 +205,8 @@ void SWaveformPanel::OnDisplayRangeUpdated(const TRange<float> NewDisplayRange)
 void SWaveformPanel::GenerateFloatRenderData()
 {
 	check(RenderData)
-
 	TArrayView<const int16> SampleData = RenderData->GetSampleData();
+
 	FloatRenderData.SetNumUninitialized(SampleData.Num());
 
 	for (int32 Sample = 0; Sample < SampleData.Num(); ++Sample)
@@ -213,7 +215,7 @@ void SWaveformPanel::GenerateFloatRenderData()
 	}
 }
 
-void SWaveformPanel::UpdateDisplayUnit(const EWaveformEditorDisplayUnit InDisplayUnit)
+void SWaveformPanel::UpdateDisplayUnit(const ESampledSequenceDisplayUnit InDisplayUnit)
 {
 	DisplayUnit = InDisplayUnit;
 	TimeRuler->UpdateDisplayUnit(DisplayUnit);
@@ -234,18 +236,18 @@ void SWaveformPanel::Tick(const FGeometry& AllottedGeometry, const double InCurr
 	}
 
 	UpdatePlayheadPosition(PaintedWidth);
-
 }
 
 void SWaveformPanel::UpdatePlayheadPosition(const float PaintedWidth)
 {
+
 	float PlayheadX = 0.f;
 
 	check(TransportCoordinator)
+
 	PlayheadX = TransportCoordinator->GetFocusPoint() * PaintedWidth;
 	PlayheadX = GridData ? GridData->SnapPositionToClosestFrame(PlayheadX) : PlayheadX;
 
-
 	if (PlayheadOverlay)
 	{
 		PlayheadOverlay->SetPlayheadPosition(PlayheadX);
@@ -275,6 +277,7 @@ FReply SWaveformPanel::HandleTimeRulerMouseMove(const FGeometry& Geometry, const
 FReply SWaveformPanel::HandleTimeRulerInteraction(const WaveformEditorPanel::EReceivedInteractionType MouseInteractionType, const FPointerEvent& MouseEvent, const FGeometry& Geometry)
 {
 	check(TimeRuler)
+
 	check(TransportCoordinator)
 
 	const float LocalWidth = Geometry.GetLocalSize().X;
@@ -283,23 +286,28 @@ FReply SWaveformPanel::HandleTimeRulerInteraction(const WaveformEditorPanel::ERe
 	{
 		const FVector2D ScreenSpacePosition = MouseEvent.GetScreenSpacePosition();
 		const FVector2D CursorPosition = Geometry.AbsoluteToLocal(ScreenSpacePosition);
+
 		const float CursorXRatio = CursorPosition.X / LocalWidth;
 
 		switch (MouseInteractionType)
 		{
 		default:
+
 			break;
+
 		case WaveformEditorPanel::EReceivedInteractionType::MouseButtonDown:
 			if (MouseEvent.IsMouseButtonDown(EKeys::LeftMouseButton))
 			{
 				return FReply::Handled().CaptureMouse(TimeRuler.ToSharedRef()).PreventThrottling();
 			}
+
 		case WaveformEditorPanel::EReceivedInteractionType::MouseMove:
 			if (MouseEvent.IsMouseButtonDown(EKeys::LeftMouseButton))
 			{
 				TransportCoordinator->ScrubFocusPoint(CursorXRatio, true);
 				return FReply::Handled().CaptureMouse(TimeRuler.ToSharedRef());
 			}
+
 		case WaveformEditorPanel::EReceivedInteractionType::MouseButtonUp:
 			if (TimeRuler->HasMouseCapture())
 			{
@@ -315,16 +323,16 @@ FReply SWaveformPanel::HandleTimeRulerInteraction(const WaveformEditorPanel::ERe
 FReply SWaveformPanel::HandlePlayheadOverlayMouseButtonUp(const FGeometry& Geometry, const FPointerEvent& MouseEvent)
 {
 	const bool HandleLeftMouseButton = MouseEvent.GetEffectingButton() == EKeys::LeftMouseButton;
+
 	if (HandleLeftMouseButton && TransportCoordinator)
 	{
-		const float LocalWidth = Geometry.GetLocalSize().X;
 
+		const float LocalWidth = Geometry.GetLocalSize().X;
 		if (LocalWidth > 0.f)
 		{
 			const float NewPosition = Geometry.AbsoluteToLocal(MouseEvent.GetScreenSpacePosition()).X / LocalWidth;
 			TransportCoordinator->ScrubFocusPoint(NewPosition, false);
 		}
 	}
-
 	return FReply::Handled();
 }
diff --git a/Engine/Plugins/Editor/WaveformEditor/Source/WaveformEditorWidgets/Private/WaveformEditorDisplayUnit.h b/Engine/Plugins/Editor/WaveformEditor/Source/WaveformEditorWidgets/Private/WaveformEditorDisplayUnit.h
deleted file mode 100644
diff --git a/Engine/Plugins/Editor/WaveformEditor/Source/WaveformEditorWidgets/Private/WaveformEditorSlateTypes.cpp b/Engine/Plugins/Editor/WaveformEditor/Source/WaveformEditorWidgets/Private/WaveformEditorSlateTypes.cpp
deleted file mode 100644
diff --git a/Engine/Plugins/Editor/WaveformEditor/Source/WaveformEditorWidgets/Private/WaveformEditorStyle.cpp b/Engine/Plugins/Editor/WaveformEditor/Source/WaveformEditorWidgets/Private/WaveformEditorStyle.cpp
--- a/Engine/Plugins/Editor/WaveformEditor/Source/WaveformEditorWidgets/Private/WaveformEditorStyle.cpp
+++ b/Engine/Plugins/Editor/WaveformEditor/Source/WaveformEditorWidgets/Private/WaveformEditorStyle.cpp
@@ -4,8 +4,6 @@
 
 #include "AudioWidgetsSlateTypes.h"
 #include "Styling/SlateStyleRegistry.h"
-#include "WaveformEditorSlateTypes.h"
-#include "WaveformEditorWidgetsSettings.h"
 #include "WaveformEditorWidgetsSettings.h"
 
 static FLazyName WaveformViewerStyleName("WaveformViewer.Style");
@@ -63,7 +61,7 @@ void FWaveformEditorStyle::Init()
 	StyleInstance->Set(PlayheadOverlayStyleName, PlayheadOverlayStyle);
 
 	//Time Ruler style 
-	FWaveformEditorTimeRulerStyle TimeRulerStyle = FWaveformEditorTimeRulerStyle()
+	FFixedSampleSequenceRulerStyle TimeRulerStyle = FFixedSampleSequenceRulerStyle()
 		.SetHandleColor(Settings->PlayheadColor)
 		.SetTicksColor(Settings->RulerTicksColor)
 		.SetTicksTextColor(Settings->RulerTextColor)
@@ -88,7 +86,7 @@ void FWaveformEditorStyle::OnWidgetSettingsUpdated(const FName& PropertyName, co
 {
 	TSharedRef<FSampledSequenceViewerStyle> WaveformViewerStyle = GetRegisteredWidgetStyle<FSampledSequenceViewerStyle>(WaveformViewerStyleName);
 	TSharedRef<FPlayheadOverlayStyle> PlayheadOverlayStyle = GetRegisteredWidgetStyle<FPlayheadOverlayStyle>(PlayheadOverlayStyleName);
-	TSharedRef<FWaveformEditorTimeRulerStyle> WaveformEditorTimeRulerStyle = GetRegisteredWidgetStyle<FWaveformEditorTimeRulerStyle>(WaveformEditorRulerStyleName);
+	TSharedRef<FFixedSampleSequenceRulerStyle> WaveformEditorTimeRulerStyle = GetRegisteredWidgetStyle<FFixedSampleSequenceRulerStyle>(WaveformEditorRulerStyleName);
 
 	if (PropertyName == GET_MEMBER_NAME_CHECKED(UWaveformEditorWidgetsSettings, WaveformColor))
 	{
diff --git a/Engine/Plugins/Editor/WaveformEditor/Source/WaveformEditorWidgets/Public/SWaveformPanel.h b/Engine/Plugins/Editor/WaveformEditor/Source/WaveformEditorWidgets/Public/SWaveformPanel.h
--- a/Engine/Plugins/Editor/WaveformEditor/Source/WaveformEditorWidgets/Public/SWaveformPanel.h
+++ b/Engine/Plugins/Editor/WaveformEditor/Source/WaveformEditorWidgets/Public/SWaveformPanel.h
@@ -5,17 +5,17 @@
 #include "HAL/Platform.h"
 #include "Widgets/SCompoundWidget.h"
 
-enum class EWaveformEditorDisplayUnit;
-class FSparseSampledSequenceTransportCoordinator;
-class FWaveformEditorGridData;
-class FWaveformEditorRenderData;
-class FWaveformEditorStyle;
-class FWaveformEditorZoomController;
+enum class ESampledSequenceDisplayUnit;
+class SWaveformViewerOverlay;
+class SWaveformTransformationsOverlay;
 class SPlayheadOverlay;
 class SFixedSampledSequenceViewer;
-class SWaveformEditorTimeRuler;
-class SWaveformTransformationsOverlay;
-class SWaveformViewerOverlay;
+class SFixedSampledSequenceRuler;
+class FWaveformEditorZoomController;
+class FWaveformEditorStyle;
+class FWaveformEditorRenderData;
+class FWaveformEditorGridData;
+class FSparseSampledSequenceTransportCoordinator;
 
 namespace WaveformEditorPanel
 {
@@ -57,19 +57,19 @@ private:
 	FReply HandleTimeRulerMouseButtonDown(const FGeometry& Geometry, const FPointerEvent& MouseEvent);
 	FReply HandleTimeRulerInteraction(const WaveformEditorPanel::EReceivedInteractionType MouseInteractionType, const FPointerEvent& MouseEvent, const FGeometry& Geometry);
 	FReply HandlePlayheadOverlayMouseButtonUp(const FGeometry& Geometry, const FPointerEvent& MouseEvent);
-	
+
 	//temporary float render data conversion while widgets are made generic
 	void GenerateFloatRenderData();
 	TArray<float> FloatRenderData;
 
-	void UpdateDisplayUnit(const EWaveformEditorDisplayUnit InDisplayUnit);
+	void UpdateDisplayUnit(const ESampledSequenceDisplayUnit InDisplayUnit);
 	void UpdatePlayheadPosition(const float PaintedWidth);
 
 	TSharedPtr<FWaveformEditorRenderData> RenderData;
 	TSharedPtr<FWaveformEditorGridData> GridData;
 	TSharedPtr<FSparseSampledSequenceTransportCoordinator> TransportCoordinator;
 
-	TSharedPtr<SWaveformEditorTimeRuler> TimeRuler;
+	TSharedPtr<SFixedSampledSequenceRuler> TimeRuler;
 	TSharedPtr<SFixedSampledSequenceViewer> WaveformViewer;
 	TSharedPtr<SWaveformTransformationsOverlay> WaveformTransformationsOverlay;
 	TSharedPtr<SWaveformViewerOverlay> WaveformViewerOverlay;
@@ -77,7 +77,7 @@ private:
 
 	float CachedPixelWidth = 0.f;
 
-	EWaveformEditorDisplayUnit DisplayUnit;
+	ESampledSequenceDisplayUnit DisplayUnit;
 
 	FWaveformEditorStyle* WaveformEditorStyle;
 };
\ No newline at end of file
diff --git a/Engine/Plugins/Editor/WaveformEditor/Source/WaveformEditorWidgets/Public/WaveformEditorSlateTypes.h b/Engine/Plugins/Editor/WaveformEditor/Source/WaveformEditorWidgets/Public/WaveformEditorSlateTypes.h
deleted file mode 100644
diff --git a/Engine/Plugins/Editor/WaveformEditor/Source/WaveformEditorWidgets/Public/WaveformEditorStyle.h b/Engine/Plugins/Editor/WaveformEditor/Source/WaveformEditorWidgets/Public/WaveformEditorStyle.h
--- a/Engine/Plugins/Editor/WaveformEditor/Source/WaveformEditorWidgets/Public/WaveformEditorStyle.h
+++ b/Engine/Plugins/Editor/WaveformEditor/Source/WaveformEditorWidgets/Public/WaveformEditorStyle.h
@@ -3,7 +3,6 @@
 #pragma once
 
 #include "Styling/SlateStyle.h"
-#include "WaveformEditorSlateTypes.h"
 
 class UWaveformEditorWidgetsSettings;
 
diff --git a/Engine/Plugins/Runtime/AudioWidgets/Source/AudioWidgets/Private/AudioWidgetsSlateTypes.cpp b/Engine/Plugins/Runtime/AudioWidgets/Source/AudioWidgets/Private/AudioWidgetsSlateTypes.cpp
--- a/Engine/Plugins/Runtime/AudioWidgets/Source/AudioWidgets/Private/AudioWidgetsSlateTypes.cpp
+++ b/Engine/Plugins/Runtime/AudioWidgets/Source/AudioWidgets/Private/AudioWidgetsSlateTypes.cpp
@@ -13,6 +13,7 @@ namespace AudioWidgetStylesSharedParams
 	const FLazyName BackgroundBrushName = "WhiteBrush";
 	const FLazyName PlayheadHandleBrushName = "Sequencer.Timeline.VanillaScrubHandleDown";
 	const FLinearColor PlayheadColor = FLinearColor(255.f, 0.1f, 0.2f, 1.f);
+	const FLinearColor RulerTicksColor = FLinearColor(1.f, 1.f, 1.f, 0.9f);
 	const float DefaultHeight = 720.f;
 	const float DefaultWidth = 1280.f;
 }
@@ -124,4 +125,33 @@ const FPlayheadOverlayStyle& FPlayheadOverlayStyle::GetDefault()
 {
 	static FPlayheadOverlayStyle Default;
 	return Default;
+}
+
+const FName FFixedSampleSequenceRulerStyle::TypeName("FFixedSampleSequenceRulerStyle");
+
+FFixedSampleSequenceRulerStyle::FFixedSampleSequenceRulerStyle()
+	: HandleWidth(15.f)
+	, HandleColor(AudioWidgetStylesSharedParams::PlayheadColor)
+	, HandleBrush(*FAppStyle::GetBrush(AudioWidgetStylesSharedParams::PlayheadHandleBrushName))
+	, TicksColor(AudioWidgetStylesSharedParams::RulerTicksColor)
+	, TicksTextColor(AudioWidgetStylesSharedParams::RulerTicksColor)
+	, TicksTextFont(FAppStyle::GetFontStyle("Regular"))
+	, TicksTextOffset(5.f)
+	, BackgroundColor(FLinearColor::Black)
+	, BackgroundBrush(*FAppStyle::GetBrush(AudioWidgetStylesSharedParams::BackgroundBrushName))
+	, DesiredWidth(AudioWidgetStylesSharedParams::DefaultWidth)
+	, DesiredHeight(30.f)
+{
+}
+
+const FFixedSampleSequenceRulerStyle& FFixedSampleSequenceRulerStyle::GetDefault()
+{
+	static FFixedSampleSequenceRulerStyle Default;
+	return Default;
+}
+
+void FFixedSampleSequenceRulerStyle::GetResources(TArray< const FSlateBrush* >& OutBrushes) const
+{
+	OutBrushes.Add(&HandleBrush);
+	OutBrushes.Add(&BackgroundBrush);
 }
\ No newline at end of file
diff --git a/Engine/Plugins/Runtime/AudioWidgets/Source/AudioWidgets/Private/SFixedSampledSequenceRuler.cpp b/Engine/Plugins/Runtime/AudioWidgets/Source/AudioWidgets/Private/SFixedSampledSequenceRuler.cpp
new file mode 100644
--- /dev/null
+++ b/Engine/Plugins/Runtime/AudioWidgets/Source/AudioWidgets/Private/SFixedSampledSequenceRuler.cpp
@@ -0,0 +1,258 @@
+// Copyright Epic Games, Inc. All Rights Reserved.
+
+#include "SFixedSampledSequenceRuler.h"
+#include "Layout/WidgetPath.h"
+#include "Framework/Application/SlateApplication.h"
+#include "Framework/MultiBox/MultiBoxBuilder.h"
+#include "SparseSampledSequenceTransportCoordinator.h"
+
+#define LOCTEXT_NAMESPACE "WaveformEditorTimeRuler"
+
+void SFixedSampledSequenceRuler::Construct(const FArguments& InArgs, TSharedRef<IFixedSampledSequenceGridService> InGridService)
+{
+	GridService = InGridService;
+	UpdateGridMetrics();
+	DisplayUnit = InArgs._DisplayUnit;
+
+	if (InArgs._DisplayPlayhead.IsSet())
+	{
+		bDisplayPlayhead = InArgs._DisplayPlayhead.Get();
+	}
+
+	Style = InArgs._Style;
+	check(Style);
+
+	HandleColor = Style->HandleColor;
+	HandleWidth = Style->HandleWidth;
+	HandleBrush = Style->HandleBrush;
+	BackgroundColor = Style->BackgroundColor;
+	BackgroundBrush = Style->BackgroundBrush;
+	TicksColor = Style->TicksColor;
+	TicksTextColor = Style->TicksTextColor;
+	TicksTextFont = Style->TicksTextFont;
+	DesiredWidth = Style->DesiredWidth;
+	DesiredHeight = Style->DesiredHeight;
+	TicksTextOffset = Style->TicksTextOffset;
+}
+
+FVector2D SFixedSampledSequenceRuler::ComputeDesiredSize(float) const
+{
+	return FVector2D(DesiredWidth, DesiredHeight);
+}
+
+int32 SFixedSampledSequenceRuler::OnPaint(const FPaintArgs& Args, const FGeometry& AllottedGeometry, const FSlateRect& MyCullingRect, FSlateWindowElementList& OutDrawElements, int32 LayerId, const FWidgetStyle& InWidgetStyle, bool bParentEnabled) const
+{
+	FSlateDrawElement::MakeBox(
+		OutDrawElements,
+		++LayerId,
+		AllottedGeometry.ToPaintGeometry(),
+		&BackgroundBrush,
+		ESlateDrawEffect::None,
+		BackgroundColor.GetSpecifiedColor());
+	
+	DrawRulerTicks(AllottedGeometry, OutDrawElements, LayerId);
+	
+	if (bDisplayPlayhead)
+	{
+		DrawPlayheadHandle(AllottedGeometry, OutDrawElements, LayerId);
+	}
+	
+	return LayerId;
+}
+
+void SFixedSampledSequenceRuler::DrawRulerTicks(const FGeometry& AllottedGeometry, FSlateWindowElementList& OutDrawElements, int32& LayerId) const
+{
+	const double MinorGridXStep = GridMetrics.MajorGridXStep / GridMetrics.NumMinorGridDivisions;
+	const double MajorTickY = AllottedGeometry.Size.Y * 0.25;
+	const double MinorTickY = AllottedGeometry.Size.Y * 0.75;
+	
+	TArray<FVector2D> LinePoints;
+	LinePoints.SetNumUninitialized(2);
+	
+	for (double CurrentMajorTickX = GridMetrics.FirstMajorTickX; CurrentMajorTickX < AllottedGeometry.Size.X; CurrentMajorTickX += GridMetrics.MajorGridXStep)
+	{
+		const double MajorTickX = CurrentMajorTickX;
+		LinePoints[0] = FVector2D(MajorTickX, MajorTickY);
+		LinePoints[1] = FVector2D(MajorTickX, AllottedGeometry.Size.Y);
+		
+		FSlateDrawElement::MakeLines(
+			OutDrawElements,
+			++LayerId,
+			AllottedGeometry.ToPaintGeometry(),
+			LinePoints,
+			ESlateDrawEffect::None,
+			TicksColor.GetSpecifiedColor(),
+			false);
+		
+		for (int32 MinorTickIndex = 1; MinorTickIndex < GridMetrics.NumMinorGridDivisions; ++MinorTickIndex)
+		{
+			const double MinorTickX = MajorTickX + MinorGridXStep * MinorTickIndex;
+			LinePoints[0] = FVector2D(MinorTickX, MinorTickY);
+			LinePoints[1] = FVector2D(MinorTickX, AllottedGeometry.Size.Y);
+			FSlateDrawElement::MakeLines(
+				OutDrawElements,
+				++LayerId,
+				AllottedGeometry.ToPaintGeometry(),
+				LinePoints,
+				ESlateDrawEffect::None,
+				TicksColor.GetSpecifiedColor(),
+				false);
+		}
+		
+		uint32 TickFrame = FMath::RoundToInt32(MajorTickX / GridMetrics.PixelsPerFrame) + GridMetrics.StartFrame;
+		DrawTickTimeString(TickFrame, MajorTickX, MajorTickY, OutDrawElements, LayerId, AllottedGeometry);
+	}
+}
+
+void SFixedSampledSequenceRuler::DrawTickTimeString(uint32 TickFrame, const double TickX, const double TickY, FSlateWindowElementList& OutDrawElements, int32& LayerId, const FGeometry& AllottedGeometry) const
+{
+	FString TimeString;
+	
+	switch (DisplayUnit)
+	{
+	case ESampledSequenceDisplayUnit::Samples:
+	{
+		uint32 SampleTime = TickFrame;
+		TimeString = FString::Printf(TEXT("%d"), SampleTime);
+	}
+	break;
+	case ESampledSequenceDisplayUnit::Seconds:
+		TimeString = FString::Printf(TEXT("%.3f"), TickFrame / (float) GridMetrics.SampleRate);
+		break;
+	default:
+		TimeString = FString::Printf(TEXT("%.3f"), TickFrame / (float)GridMetrics.SampleRate);
+		break;
+	}
+	
+	FVector2D TextOffset(TickX + TicksTextOffset, TickY);
+	
+	FSlateDrawElement::MakeText(
+		OutDrawElements,
+		++LayerId,
+		AllottedGeometry.ToPaintGeometry(AllottedGeometry.Size, FSlateLayoutTransform(TextOffset)),
+		TimeString,
+		TicksTextFont,
+		ESlateDrawEffect::None,
+		TicksTextColor.GetSpecifiedColor()
+	);
+}
+
+TSharedRef<SWidget> SFixedSampledSequenceRuler::MakeContextMenu()
+{
+	const bool bCloseAfterSelection = false;
+	const bool bCloseSelfOnly = false;
+	FMenuBuilder MenuBuilder(bCloseAfterSelection, NULL, TSharedPtr<FExtender>(), bCloseSelfOnly, &FCoreStyle::Get());
+	MenuBuilder.BeginSection("Ruler Options", LOCTEXT("RulerOptionsHeading", "Options"));
+	{
+		MenuBuilder.AddSubMenu(
+			LOCTEXT("TimeUnitsSubMenuLabel", "Time Unit"),
+			LOCTEXT("TimeUnitsSubMenuLabelToolTip", "Select the time unit shown by the grid"),
+			FNewMenuDelegate::CreateSP(this, &SFixedSampledSequenceRuler::MakeTimeUnitsSubMenu)
+		);
+	}
+	MenuBuilder.EndSection();
+	return MenuBuilder.MakeWidget();
+}
+
+void SFixedSampledSequenceRuler::MakeTimeUnitsSubMenu(FMenuBuilder& SubMenuBuilder)
+{
+	SubMenuBuilder.BeginSection(NAME_None);
+
+	SubMenuBuilder.AddMenuEntry(
+		LOCTEXT("TimeUnitsSubMenuEntry_Seconds", "Seconds"),
+		LOCTEXT("TimeUnitsSubMenuEntry_Seconds", "Seconds"),
+		FSlateIcon(),
+		FUIAction(
+			FExecuteAction::CreateSP(this, &SFixedSampledSequenceRuler::NotifyTimeUnitMenuSelection, ESampledSequenceDisplayUnit::Seconds),
+			FCanExecuteAction::CreateLambda([this] { return DisplayUnit != ESampledSequenceDisplayUnit::Seconds; }),
+			FIsActionChecked::CreateLambda([this] { return DisplayUnit == ESampledSequenceDisplayUnit::Seconds;})
+		),
+		NAME_None,
+		EUserInterfaceActionType::ToggleButton
+	);
+
+	SubMenuBuilder.AddMenuEntry(
+		LOCTEXT("TimeUnitsSubMenuEntry_Frames", "Frames"),
+		LOCTEXT("TimeUnitsSubMenuEntry_Frames", "Frames"),
+		FSlateIcon(),
+		FUIAction(
+			FExecuteAction::CreateSP(this, &SFixedSampledSequenceRuler::NotifyTimeUnitMenuSelection, ESampledSequenceDisplayUnit::Samples),
+			FCanExecuteAction::CreateLambda([this] { return DisplayUnit != ESampledSequenceDisplayUnit::Samples; }),
+			FIsActionChecked::CreateLambda([this] { return DisplayUnit == ESampledSequenceDisplayUnit::Samples; })
+		),
+		NAME_None,
+		EUserInterfaceActionType::ToggleButton
+	);
+
+	SubMenuBuilder.EndSection();
+ }
+
+void SFixedSampledSequenceRuler::NotifyTimeUnitMenuSelection(const ESampledSequenceDisplayUnit SelectedDisplayUnit) const
+{
+	if (SelectedDisplayUnit == DisplayUnit)
+	{
+		return;
+	}
+
+	if (OnTimeUnitMenuSelection.IsBound())
+	{
+		OnTimeUnitMenuSelection.Broadcast(SelectedDisplayUnit);
+	}
+}
+
+void SFixedSampledSequenceRuler::UpdateGridMetrics()
+{
+	check(GridService)
+	GridMetrics = GridService->GetGridMetrics();
+}
+
+void SFixedSampledSequenceRuler::UpdateDisplayUnit(const ESampledSequenceDisplayUnit InDisplayUnit)
+{
+	DisplayUnit = InDisplayUnit;
+}
+
+void SFixedSampledSequenceRuler::SetPlayheadPosition(const float InNewPosition)
+{
+	PlayheadPosition = InNewPosition;
+}
+
+void SFixedSampledSequenceRuler::OnStyleUpdated(const FNotifyingAudioWidgetStyle& UpdatedStyle)
+{
+	check(Style);
+	if (&UpdatedStyle != Style)
+	{
+		return;
+	}
+
+	HandleColor = Style->HandleColor;
+	HandleWidth = Style->HandleWidth;
+	HandleBrush = Style->HandleBrush;
+	BackgroundColor = Style->BackgroundColor;
+	BackgroundBrush = Style->BackgroundBrush;
+	TicksColor = Style->TicksColor;
+	TicksTextColor = Style->TicksTextColor;
+	TicksTextFont = Style->TicksTextFont;
+	DesiredWidth = Style->DesiredWidth;
+	DesiredHeight = Style->DesiredHeight;
+	TicksTextOffset = Style->TicksTextOffset;
+}
+
+void SFixedSampledSequenceRuler::DrawPlayheadHandle(const FGeometry& AllottedGeometry, FSlateWindowElementList& OutDrawElements, int32& LayerId) const
+{
+	const float WindowWidth = AllottedGeometry.Size.X;
+	const float PlayheadCenter = PlayheadPosition;
+	const float	HandleStart = PlayheadCenter - HandleWidth / 2;
+	const float	HandleEnd = HandleStart + HandleWidth;
+ 	FPaintGeometry HandleGeometry = AllottedGeometry.ToPaintGeometry(FVector2f(HandleEnd - HandleStart, AllottedGeometry.Size.Y), FSlateLayoutTransform(FVector2f(HandleStart, 0)));
+	FSlateDrawElement::MakeBox(
+		OutDrawElements,
+		LayerId,
+		HandleGeometry,
+		&HandleBrush,
+		ESlateDrawEffect::None,
+		HandleColor.GetSpecifiedColor()
+	);
+	++LayerId;
+}
+
+#undef LOCTEXT_NAMESPACE
\ No newline at end of file
diff --git a/Engine/Plugins/Runtime/AudioWidgets/Source/AudioWidgets/Public/AudioWidgetsSlateTypes.h b/Engine/Plugins/Runtime/AudioWidgets/Source/AudioWidgets/Public/AudioWidgetsSlateTypes.h
--- a/Engine/Plugins/Runtime/AudioWidgets/Source/AudioWidgets/Public/AudioWidgetsSlateTypes.h
+++ b/Engine/Plugins/Runtime/AudioWidgets/Source/AudioWidgets/Public/AudioWidgetsSlateTypes.h
@@ -279,6 +279,68 @@ struct AUDIOWIDGETS_API FPlayheadOverlayStyle : public FNotifyingAudioWidgetStyl
 	FPlayheadOverlayStyle& SetDesiredHeight(const float InDesiredHeight) { DesiredHeight = InDesiredHeight; BroadcastStyleUpdate(); return *this; }
 };
 
+/**
+ * Represents the appearance of a Sampled Sequence Time Ruler
+ */
+USTRUCT(BlueprintType)
+struct AUDIOWIDGETS_API FFixedSampleSequenceRulerStyle : public FNotifyingAudioWidgetStyle
+{
+	GENERATED_USTRUCT_BODY()
+
+	FFixedSampleSequenceRulerStyle();
+
+	static const FFixedSampleSequenceRulerStyle& GetDefault();
+	virtual const FName GetTypeName() const override { return TypeName; };
+	virtual void GetResources(TArray< const FSlateBrush* >& OutBrushes) const override;
+
+	static const FName TypeName;
+
+	UPROPERTY(EditAnywhere, BlueprintReadWrite, Category = Appearance)
+	float HandleWidth;
+	FFixedSampleSequenceRulerStyle& SetHandleWidth(const float InHandleWidth) { HandleWidth = InHandleWidth; BroadcastStyleUpdate(); return *this; }
+
+	UPROPERTY(EditAnywhere, BlueprintReadWrite, Category = Appearance)
+	FSlateColor HandleColor;
+	FFixedSampleSequenceRulerStyle& SetHandleColor(const FSlateColor& InHandleColor) { HandleColor = InHandleColor; BroadcastStyleUpdate(); return *this; }
+
+	UPROPERTY(EditAnywhere, BlueprintReadWrite, Category = Appearance)
+	FSlateBrush HandleBrush;
+	FFixedSampleSequenceRulerStyle& SetHandleBrush(const FSlateBrush& InHandleBrush) { HandleBrush = InHandleBrush; BroadcastStyleUpdate(); return *this; }
+
+	UPROPERTY(EditAnywhere, BlueprintReadWrite, Category = Appearance)
+	FSlateColor TicksColor;
+	FFixedSampleSequenceRulerStyle& SetTicksColor(const FSlateColor& InTicksColor) { TicksColor = InTicksColor; BroadcastStyleUpdate(); return *this; }
+
+	UPROPERTY(EditAnywhere, BlueprintReadWrite, Category = Appearance)
+	FSlateColor TicksTextColor;
+	FFixedSampleSequenceRulerStyle& SetTicksTextColor(const FSlateColor& InTicksTextColor) { TicksTextColor = InTicksTextColor; BroadcastStyleUpdate(); return *this; }
+
+	UPROPERTY(EditAnywhere, BlueprintReadWrite, Category = Appearance)
+	FSlateFontInfo TicksTextFont;
+	FFixedSampleSequenceRulerStyle& SetTicksTextFont(const FSlateFontInfo& InTicksTextFont) { TicksTextFont = InTicksTextFont; BroadcastStyleUpdate(); return *this; }
+	FFixedSampleSequenceRulerStyle& SetFontSize(const float InFontSize) { TicksTextFont.Size = InFontSize; BroadcastStyleUpdate(); return *this; }
+
+	UPROPERTY(EditAnywhere, BlueprintReadWrite, Category = Appearance)
+	float TicksTextOffset;
+	FFixedSampleSequenceRulerStyle& SetTicksTextOffset(const float InTicksTextOffset) { TicksTextOffset = InTicksTextOffset; BroadcastStyleUpdate(); return *this; }
+
+	UPROPERTY(EditAnywhere, BlueprintReadWrite, Category = Appearance)
+	FSlateColor BackgroundColor;
+	FFixedSampleSequenceRulerStyle& SetBackgroundColor(const FSlateColor& InBackgroundColor) { BackgroundColor = InBackgroundColor; BroadcastStyleUpdate(); return *this; }
+
+	UPROPERTY(EditAnywhere, BlueprintReadWrite, Category = Appearance)
+	FSlateBrush BackgroundBrush;
+	FFixedSampleSequenceRulerStyle& SetBackgroundBrush(const FSlateBrush& InBackgroundBrush) { BackgroundBrush = InBackgroundBrush; BroadcastStyleUpdate(); return *this; }
+
+	UPROPERTY(EditAnywhere, BlueprintReadWrite, Category = Appearance)
+	float DesiredWidth;
+	FFixedSampleSequenceRulerStyle& SetDesiredWidth(const float InDesiredWidth) { DesiredWidth = InDesiredWidth; BroadcastStyleUpdate(); return *this; }
+
+	UPROPERTY(EditAnywhere, BlueprintReadWrite, Category = Appearance)
+	float DesiredHeight;
+	FFixedSampleSequenceRulerStyle& SetDesiredHeight(const float InDesiredHeight) { DesiredHeight = InDesiredHeight; BroadcastStyleUpdate(); return *this; }
+};
+
 #if UE_ENABLE_INCLUDE_ORDER_DEPRECATED_IN_5_2
 #include "Styling/SlateStyle.h"
 #include "Styling/SlateStyleMacros.h"
diff --git a/Engine/Plugins/Runtime/AudioWidgets/Source/AudioWidgets/Public/SFixedSampledSequenceRuler.h b/Engine/Plugins/Runtime/AudioWidgets/Source/AudioWidgets/Public/SFixedSampledSequenceRuler.h
new file mode 100644
--- /dev/null
+++ b/Engine/Plugins/Runtime/AudioWidgets/Source/AudioWidgets/Public/SFixedSampledSequenceRuler.h
@@ -0,0 +1,77 @@
+// Copyright Epic Games, Inc. All Rights Reserved.
+
+#pragma once
+
+#include "IFixedSampledSequenceGridService.h"
+#include "Framework/MultiBox/MultiBoxBuilder.h"
+#include "Styling/SlateStyleRegistry.h"
+#include "Styling/SlateWidgetStyleAsset.h"
+#include "SampledSequenceDisplayUnit.h"
+#include "AudioWidgetsSlateTypes.h"
+#include "Widgets/SCompoundWidget.h"
+
+DECLARE_MULTICAST_DELEGATE_OneParam(FOnTimeUnitMenuSelection, const ESampledSequenceDisplayUnit /* Requested Display Unit */);
+
+class AUDIOWIDGETS_API SFixedSampledSequenceRuler : public SCompoundWidget
+{
+	SLATE_BEGIN_ARGS(SFixedSampledSequenceRuler)
+		: _DisplayUnit(ESampledSequenceDisplayUnit::Seconds)
+	{
+	}
+
+	SLATE_ARGUMENT(ESampledSequenceDisplayUnit, DisplayUnit)
+
+	/** Whether the playhead should be drawn or not */
+	SLATE_ATTRIBUTE(bool, DisplayPlayhead)
+
+	SLATE_STYLE_ARGUMENT(FFixedSampleSequenceRulerStyle, Style)
+
+	SLATE_END_ARGS()
+
+public:
+	void Construct(const FArguments& InArgs, TSharedRef<IFixedSampledSequenceGridService> InGridService);
+	void UpdateGridMetrics();
+	void UpdateDisplayUnit(const ESampledSequenceDisplayUnit InDisplayUnit);
+	void SetPlayheadPosition(const float InNewPosition);
+	void OnStyleUpdated(const FNotifyingAudioWidgetStyle& UpdatedStyle);
+
+	/** Delegate sent when the user selects a new display unit from the RMB menu*/
+	FOnTimeUnitMenuSelection OnTimeUnitMenuSelection;
+
+private:
+	virtual int32 OnPaint(const FPaintArgs& Args, const FGeometry& AllottedGeometry, const FSlateRect& MyCullingRect, FSlateWindowElementList& OutDrawElements, int32 LayerId, const FWidgetStyle& InWidgetStyle, bool bParentEnabled) const override;
+	virtual FVector2D ComputeDesiredSize(float) const override;
+
+	void DrawPlayheadHandle(const FGeometry& AllottedGeometry, FSlateWindowElementList& OutDrawElements, int32& LayerId) const;
+	void DrawRulerTicks(const FGeometry& AllottedGeometry, FSlateWindowElementList& OutDrawElements, int32& LayerId) const;
+	void DrawTickTimeString(uint32 TickFrame, const double TickX, const double TickY, FSlateWindowElementList& OutDrawElements, int32& LayerId, const FGeometry& AllottedGeometry) const;
+
+	TSharedRef<SWidget> MakeContextMenu();
+	void MakeTimeUnitsSubMenu(FMenuBuilder& SubMenuBuilder);
+	void NotifyTimeUnitMenuSelection(const ESampledSequenceDisplayUnit SelectedDisplayUnit) const;
+
+	FFixedSampledSequenceGridMetrics GridMetrics;
+
+	const FFixedSampleSequenceRulerStyle* Style = nullptr;
+
+	FSlateColor BackgroundColor = FLinearColor::Black;
+	FSlateBrush BackgroundBrush;
+	FSlateBrush HandleBrush;
+	FSlateColor HandleColor = FLinearColor(255.f, 0.1f, 0.2f, 1.f);
+	FSlateColor TicksColor = FLinearColor(1.f, 1.f, 1.f, 0.9f);
+	FSlateColor TicksTextColor = FLinearColor(1.f, 1.f, 1.f, 0.9f);
+
+	float DesiredHeight = 0.f;
+	float DesiredWidth = 0.f;
+	float HandleWidth = 15.f;
+	float TicksTextOffset = 5.f;
+	float PlayheadPosition = 0.f;
+
+	FSlateFontInfo TicksTextFont;
+
+	TSharedPtr<IFixedSampledSequenceGridService> GridService = nullptr;
+
+	ESampledSequenceDisplayUnit DisplayUnit;
+
+	bool bDisplayPlayhead = true;
+};
\ No newline at end of file
diff --git a/Engine/Plugins/Runtime/AudioWidgets/Source/AudioWidgets/Public/SampledSequenceDisplayUnit.h b/Engine/Plugins/Runtime/AudioWidgets/Source/AudioWidgets/Public/SampledSequenceDisplayUnit.h
new file mode 100644
--- /dev/null
+++ b/Engine/Plugins/Runtime/AudioWidgets/Source/AudioWidgets/Public/SampledSequenceDisplayUnit.h
@@ -0,0 +1,9 @@
+// Copyright Epic Games, Inc. All Rights Reserved.
+
+#pragma once
+
+enum class ESampledSequenceDisplayUnit
+{
+	Samples,
+	Seconds
+};
\ No newline at end of file
