[WaveformEditor] Rename SWaveformViewerOverlay to SWaveformEditorInputRoutingOverlay #trivial
#trivial

diff --git a/Engine/Plugins/Editor/WaveformEditor/Source/WaveformEditorWidgets/Private/SWaveformEditorInputRoutingOverlay.cpp b/Engine/Plugins/Editor/WaveformEditor/Source/WaveformEditorWidgets/Private/SWaveformEditorInputRoutingOverlay.cpp
new file mode 100644
--- /dev/null
+++ b/Engine/Plugins/Editor/WaveformEditor/Source/WaveformEditorWidgets/Private/SWaveformEditorInputRoutingOverlay.cpp
@@ -0,0 +1,62 @@
+// Copyright Epic Games, Inc. All Rights Reserved.
+
+#include "SWaveformEditorInputRoutingOverlay.h"
+
+#include "AudioWidgetsUtils.h"
+#include "Widgets/SLeafWidget.h"
+
+void SWaveformEditorInputRoutingOverlay::Construct(const FArguments& InArgs, const TArray<TSharedPtr<SWidget>>& InOverlaidWidgets)
+{
+	if (InArgs._Style)
+	{
+		Style = InArgs._Style;
+		DesiredWidth = Style->DesiredWidth;
+		DesiredHeight = Style->DesiredHeight;
+	}
+	
+	OverlaidWidgets = InOverlaidWidgets;
+}
+
+FReply SWaveformEditorInputRoutingOverlay::OnMouseButtonDown(const FGeometry& MyGeometry, const FPointerEvent& MouseEvent)
+{
+	return AudioWidgetsUtils::RouteMouseInput(&SWidget::OnMouseButtonDown, MouseEvent, MakeArrayView(OverlaidWidgets.GetData(), OverlaidWidgets.Num()));
+}
+
+FReply SWaveformEditorInputRoutingOverlay::OnMouseButtonUp(const FGeometry& MyGeometry, const FPointerEvent& MouseEvent)
+{
+	return AudioWidgetsUtils::RouteMouseInput(&SWidget::OnMouseButtonUp, MouseEvent, MakeArrayView(OverlaidWidgets.GetData(), OverlaidWidgets.Num()));
+}
+
+FReply SWaveformEditorInputRoutingOverlay::OnMouseMove(const FGeometry& MyGeometry, const FPointerEvent& MouseEvent)
+{
+	return AudioWidgetsUtils::RouteMouseInput(&SWidget::OnMouseMove, MouseEvent, MakeArrayView(OverlaidWidgets.GetData(), OverlaidWidgets.Num()));
+}
+
+FReply SWaveformEditorInputRoutingOverlay::OnMouseWheel(const FGeometry& MyGeometry, const FPointerEvent& MouseEvent)
+{
+	FReply InteractionReply = AudioWidgetsUtils::RouteMouseInput(&SWidget::OnMouseWheel, MouseEvent, MakeArrayView(OverlaidWidgets.GetData(), OverlaidWidgets.Num()));
+
+	if (!InteractionReply.IsEventHandled())
+	{
+		OnNewMouseDelta.ExecuteIfBound(MouseEvent.GetWheelDelta());
+		InteractionReply = FReply::Handled();
+	}
+
+	return InteractionReply;
+
+}
+
+FCursorReply SWaveformEditorInputRoutingOverlay::OnCursorQuery(const FGeometry& MyGeometry, const FPointerEvent& CursorEvent) const
+{
+	return AudioWidgetsUtils::RouteCursorQuery(CursorEvent, MakeArrayView(OverlaidWidgets.GetData(), OverlaidWidgets.Num()));
+}
+
+int32 SWaveformEditorInputRoutingOverlay::OnPaint(const FPaintArgs& Args, const FGeometry& AllottedGeometry, const FSlateRect& MyCullingRect, FSlateWindowElementList& OutDrawElements, int32 LayerId, const FWidgetStyle& InWidgetStyle, bool bParentEnabled) const
+{
+	return ++LayerId;
+}
+
+FVector2D SWaveformEditorInputRoutingOverlay::ComputeDesiredSize(float) const
+{
+	return FVector2D(DesiredWidth, DesiredHeight);
+}
\ No newline at end of file
diff --git a/Engine/Plugins/Editor/WaveformEditor/Source/WaveformEditorWidgets/Private/SWaveformEditorInputRoutingOverlay.h b/Engine/Plugins/Editor/WaveformEditor/Source/WaveformEditorWidgets/Private/SWaveformEditorInputRoutingOverlay.h
new file mode 100644
--- /dev/null
+++ b/Engine/Plugins/Editor/WaveformEditor/Source/WaveformEditorWidgets/Private/SWaveformEditorInputRoutingOverlay.h
@@ -0,0 +1,46 @@
+// Copyright Epic Games, Inc. All Rights Reserved.
+
+#pragma once
+
+#include "AudioWidgetsSlateTypes.h"
+#include "Styling/ISlateStyle.h"
+#include "Styling/SlateWidgetStyleAsset.h"
+#include "Widgets/SLeafWidget.h"
+
+DECLARE_DELEGATE_OneParam(FOnNewMouseDelta, const float /* new delta */)
+
+class FWaveformEditorGridData;
+class FSparseSampledSequenceTransportCoordinator;
+
+class SWaveformEditorInputRoutingOverlay : public SLeafWidget
+{
+public:
+	SLATE_BEGIN_ARGS(SWaveformEditorInputRoutingOverlay)
+	{
+	}
+
+	SLATE_STYLE_ARGUMENT(FSampledSequenceViewerStyle, Style)
+
+	SLATE_END_ARGS()
+
+	void Construct(const FArguments& InArgs, const TArray<TSharedPtr<SWidget>>& InOverlaidWidgets);
+
+	virtual FReply OnMouseButtonDown(const FGeometry& MyGeometry, const FPointerEvent& MouseEvent) override;
+	virtual FReply OnMouseButtonUp(const FGeometry& MyGeometry, const FPointerEvent& MouseEvent)override;
+	virtual FReply OnMouseWheel(const FGeometry& MyGeometry, const FPointerEvent& MouseEvent) override;
+	virtual FReply OnMouseMove(const FGeometry& MyGeometry, const FPointerEvent& MouseEvent) override;
+	virtual FCursorReply OnCursorQuery(const FGeometry& MyGeometry, const FPointerEvent& CursorEvent) const override;
+
+	FOnNewMouseDelta OnNewMouseDelta;\
+
+private:
+	virtual int32 OnPaint(const FPaintArgs& Args, const FGeometry& AllottedGeometry, const FSlateRect& MyCullingRect, FSlateWindowElementList& OutDrawElements, int32 LayerId, const FWidgetStyle& InWidgetStyle, bool bParentEnabled) const override;
+	virtual FVector2D ComputeDesiredSize(float) const override;
+
+	TArray<TSharedPtr<SWidget>> OverlaidWidgets;
+
+	const FSampledSequenceViewerStyle* Style = nullptr;
+
+	float DesiredWidth = 1280.f;
+	float DesiredHeight = 720.f;
+};
\ No newline at end of file
diff --git a/Engine/Plugins/Editor/WaveformEditor/Source/WaveformEditorWidgets/Private/SWaveformPanel.cpp b/Engine/Plugins/Editor/WaveformEditor/Source/WaveformEditorWidgets/Private/SWaveformPanel.cpp
--- a/Engine/Plugins/Editor/WaveformEditor/Source/WaveformEditorWidgets/Private/SWaveformPanel.cpp
+++ b/Engine/Plugins/Editor/WaveformEditor/Source/WaveformEditorWidgets/Private/SWaveformPanel.cpp
@@ -8,7 +8,7 @@
 #include "SparseSampledSequenceTransportCoordinator.h"
 #include "SPlayheadOverlay.h"
 #include "SWaveformTransformationsOverlay.h"
-#include "SWaveformViewerOverlay.h"
+#include "SWaveformEditorInputRoutingOverlay.h"
 #include "WaveformEditorGridData.h"
 #include "WaveformEditorRenderData.h"
 #include "WaveformEditorStyle.h"
@@ -43,7 +43,7 @@ void SWaveformPanel::Construct(const FArguments& InArgs, TSharedRef<FWaveformEdi
 	}
 
 	SetupPlayheadOverlay();
-	SetUpWaveformViewerOverlay(InZoomManager);
+	SetUpInputRoutingOverlay(InZoomManager);
 	SetUpTimeRuler(GridData.ToSharedRef());
 	
 	CreateLayout();
@@ -53,7 +53,7 @@ void SWaveformPanel::CreateLayout()
 {
 	check(TimeRuler);
 	check(WaveformViewer);
-	check(WaveformViewerOverlay);
+	check(InputRoutingOverlay);
 
 	TSharedPtr<SOverlay> WaveformView = SNew(SOverlay);
 
@@ -77,7 +77,7 @@ void SWaveformPanel::CreateLayout()
 
 	WaveformView->AddSlot()
 	[
-		WaveformViewerOverlay.ToSharedRef()
+		InputRoutingOverlay.ToSharedRef()
 	];
 
 	ChildSlot
@@ -113,7 +113,7 @@ void SWaveformPanel::SetUpTimeRuler(TSharedRef<FWaveformEditorGridData> InGridDa
 	TimeRuler->SetOnMouseMove(FPointerEventHandler::CreateSP(this, &SWaveformPanel::HandleTimeRulerMouseMove));
 }
 
-void SWaveformPanel::SetUpWaveformViewerOverlay(TSharedRef<FWaveformEditorZoomController> InZoomManager)
+void SWaveformPanel::SetUpInputRoutingOverlay(TSharedRef<FWaveformEditorZoomController> InZoomManager)
 {
 	FSampledSequenceViewerStyle* ViewerStyle = &WaveformEditorStyle->GetRegisteredWidgetStyle<FSampledSequenceViewerStyle>("WaveformViewer.Style").Get();
 	check(ViewerStyle);
@@ -127,8 +127,8 @@ void SWaveformPanel::SetUpWaveformViewerOverlay(TSharedRef<FWaveformEditorZoomCo
 
 	check(PlayheadOverlay)
 	OverlaidWidgets.Add(PlayheadOverlay);
-	WaveformViewerOverlay = SNew(SWaveformViewerOverlay, OverlaidWidgets).Style(ViewerStyle);
-	WaveformViewerOverlay->OnNewMouseDelta.BindSP(InZoomManager, &FWaveformEditorZoomController::ZoomByDelta);
+	InputRoutingOverlay = SNew(SWaveformEditorInputRoutingOverlay, OverlaidWidgets).Style(ViewerStyle);
+	InputRoutingOverlay->OnNewMouseDelta.BindSP(InZoomManager, &FWaveformEditorZoomController::ZoomByDelta);
 }
 
 void SWaveformPanel::SetupPlayheadOverlay()
diff --git a/Engine/Plugins/Editor/WaveformEditor/Source/WaveformEditorWidgets/Private/SWaveformViewerOverlay.cpp b/Engine/Plugins/Editor/WaveformEditor/Source/WaveformEditorWidgets/Private/SWaveformViewerOverlay.cpp
deleted file mode 100644
diff --git a/Engine/Plugins/Editor/WaveformEditor/Source/WaveformEditorWidgets/Private/SWaveformViewerOverlay.h b/Engine/Plugins/Editor/WaveformEditor/Source/WaveformEditorWidgets/Private/SWaveformViewerOverlay.h
deleted file mode 100644
diff --git a/Engine/Plugins/Editor/WaveformEditor/Source/WaveformEditorWidgets/Public/SWaveformPanel.h b/Engine/Plugins/Editor/WaveformEditor/Source/WaveformEditorWidgets/Public/SWaveformPanel.h
--- a/Engine/Plugins/Editor/WaveformEditor/Source/WaveformEditorWidgets/Public/SWaveformPanel.h
+++ b/Engine/Plugins/Editor/WaveformEditor/Source/WaveformEditorWidgets/Public/SWaveformPanel.h
@@ -6,7 +6,7 @@
 #include "Widgets/SCompoundWidget.h"
 
 enum class ESampledSequenceDisplayUnit;
-class SWaveformViewerOverlay;
+class SWaveformEditorInputRoutingOverlay;
 class SWaveformTransformationsOverlay;
 class SPlayheadOverlay;
 class SFixedSampledSequenceViewer;
@@ -45,7 +45,7 @@ private:
 	void SetupPlayheadOverlay();
 	void SetUpTimeRuler(TSharedRef<FWaveformEditorGridData> InGridData);
 	void SetUpWaveformViewer(TSharedRef<FWaveformEditorGridData> InGridData, TSharedRef<FWaveformEditorRenderData> InRenderData);
-	void SetUpWaveformViewerOverlay(TSharedRef<FWaveformEditorZoomController> InZoomManager);
+	void SetUpInputRoutingOverlay(TSharedRef<FWaveformEditorZoomController> InZoomManager);
 	void SetUpZoomManager(TSharedRef<FWaveformEditorZoomController> InZoomManager, TSharedRef<FSparseSampledSequenceTransportCoordinator> InTransportCoordinator);
 
 	void OnRenderDataUpdated();
@@ -72,7 +72,7 @@ private:
 	TSharedPtr<SFixedSampledSequenceRuler> TimeRuler;
 	TSharedPtr<SFixedSampledSequenceViewer> WaveformViewer;
 	TSharedPtr<SWaveformTransformationsOverlay> WaveformTransformationsOverlay;
-	TSharedPtr<SWaveformViewerOverlay> WaveformViewerOverlay;
+	TSharedPtr<SWaveformEditorInputRoutingOverlay> InputRoutingOverlay;
 	TSharedPtr<SPlayheadOverlay> PlayheadOverlay;
 
 	float CachedPixelWidth = 0.f;
