[WaveformEditor] Major Refactor to remove the use of RenderData
- Waveform Panel and TransformationsRenderManagaer are now created from a factory, and follow a data provider -> receiver model
- Waveform Panel and child widgets are unaware of Data Provider, they simply receive data and display it. User input is handled through external method injection.
- Data to render is owned by the producer
- Small optims on how the wave data to render is produced

diff --git a/Engine/Plugins/Editor/WaveformEditor/Source/WaveformEditor/Private/TransformedWaveformViewFactory.cpp b/Engine/Plugins/Editor/WaveformEditor/Source/WaveformEditor/Private/TransformedWaveformViewFactory.cpp
new file mode 100644
--- /dev/null
+++ b/Engine/Plugins/Editor/WaveformEditor/Source/WaveformEditor/Private/TransformedWaveformViewFactory.cpp
@@ -0,0 +1,242 @@
+// Copyright Epic Games, Inc. All Rights Reserved.
+
+#include "TransformedWaveformViewFactory.h"
+
+#include "Misc/NotifyHook.h"
+#include "PropertyEditorModule.h"
+#include "Sound/SoundWave.h"
+#include "SparseSampledSequenceTransportCoordinator.h"
+#include "SWaveformTransformationsOverlay.h"
+#include "TransformedWaveformView.h"
+#include "TransformedWaveformView.h"
+#include "WaveformEditorDetailsCustomization.h"
+#include "WaveformEditorSequenceDataProvider.h"
+#include "WaveformEditorZoomController.h"
+
+static FLazyName TimeRulerWidgetName("SFixedSampledSequenceRuler");
+
+FTransformedWaveformViewFactory& FTransformedWaveformViewFactory::Get()
+{
+	check(Instance);
+	return *Instance;
+}
+
+void FTransformedWaveformViewFactory::Create()
+{
+	if (Instance == nullptr)
+	{
+		Instance = MakeUnique<FTransformedWaveformViewFactory>();
+	}
+}
+
+FTransformedWaveformView FTransformedWaveformViewFactory::GetTransformedView(TObjectPtr<USoundWave> SoundWaveToView, TSharedRef<FSparseSampledSequenceTransportCoordinator> TransportCoordinator, FNotifyHook* OwnerNotifyHook /*= nullptr*/, TSharedPtr<FWaveformEditorZoomController> InZoomController /*= nullptr*/)
+{
+	FOnTransformationsPropertiesRequired OnTransformationPropertiesRequired = OwnerNotifyHook ? SetUpOwnerNotifyHook(SoundWaveToView, OwnerNotifyHook) : FOnTransformationsPropertiesRequired();
+
+	TSharedPtr<FWaveformEditorSequenceDataProvider> WaveformDataProvider = MakeShared<FWaveformEditorSequenceDataProvider>(SoundWaveToView, OnTransformationPropertiesRequired);
+	WaveformDataProvider->UpdateRenderElements();
+
+	FFixedSampledSequenceView SequenceView = WaveformDataProvider->RequestSequenceView(TRange<double>::Inclusive(0, 1));
+
+	TSharedPtr<SWaveformTransformationsOverlay> TransformationsOverlay = SNew(SWaveformTransformationsOverlay, WaveformDataProvider->GetTransformLayers(), TransportCoordinator);
+
+	TSharedPtr<STransformedWaveformViewPanel> WaveformPanel = SNew(STransformedWaveformViewPanel, SequenceView).TransformationsOverlay(TransformationsOverlay);
+	SetUpWaveformPanelInteractions(WaveformPanel.ToSharedRef(), TransportCoordinator, InZoomController, TransformationsOverlay);
+
+	WaveformDataProvider->OnLayersChainGenerated.AddSP(TransformationsOverlay.Get(), &SWaveformTransformationsOverlay::OnLayerChainGenerated);
+	WaveformDataProvider->OnRenderElementsUpdated.AddSP(TransformationsOverlay.Get(), &SWaveformTransformationsOverlay::UpdateLayerConstraints);
+	TransportCoordinator->OnDisplayRangeUpdated.AddSP(TransformationsOverlay.Get(), &SWaveformTransformationsOverlay::OnNewWaveformDisplayRange);
+	WaveformDataProvider->OnDataViewGenerated.AddSP(WaveformPanel.Get(), &STransformedWaveformViewPanel::ReceiveSequenceView);
+
+	return FTransformedWaveformView{ WaveformPanel, WaveformDataProvider};
+
+}
+
+FOnTransformationsPropertiesRequired FTransformedWaveformViewFactory::SetUpOwnerNotifyHook(TObjectPtr<USoundWave> SoundWaveToView, FNotifyHook* InNotifyHook)
+{
+	TSharedPtr<IDetailsView> PropertiesPropagator = SetUpTransformationsPropertiesPropagator(SoundWaveToView, InNotifyHook);
+
+	FOnTransformationsPropertiesRequired OnTransformationPropertiesRequired = FOnTransformationsPropertiesRequired::CreateLambda([this, SoundWaveToView, PropertiesPropagator](FTransformationsToPropertiesArray& InObjToPropsMap)
+	{
+		TSharedPtr<FWaveformTransformationsDetailsProvider> TransformationsDetailsProvider = MakeShared<FWaveformTransformationsDetailsProvider>();
+
+		FOnGetDetailCustomizationInstance ProviderInstance = FOnGetDetailCustomizationInstance::CreateLambda([&]()
+			{
+				return TransformationsDetailsProvider.ToSharedRef();
+			}
+		);
+
+		PropertiesPropagator->RegisterInstancedCustomPropertyLayout(SoundWaveToView->GetClass(), ProviderInstance);
+		PropertiesPropagator->ForceRefresh();
+
+		for (FTransformationToPropertiesPair& ObjToPropsPair : InObjToPropsMap)
+		{
+			TransformationsDetailsProvider->GetHandlesForUObjectProperties(ObjToPropsPair.Key, ObjToPropsPair.Value);
+		}
+
+		TransformationsDetailsProvider.Reset();
+	});
+
+	return OnTransformationPropertiesRequired;
+
+}
+
+TSharedPtr<IDetailsView> FTransformedWaveformViewFactory::SetUpTransformationsPropertiesPropagator(TObjectPtr<USoundWave> SoundWaveToView, FNotifyHook* InNotifyHook)
+{
+	check(InNotifyHook)
+
+	FPropertyEditorModule& PropertyModule = FModuleManager::LoadModuleChecked<FPropertyEditorModule>("PropertyEditor");
+	FDetailsViewArgs Args;
+	Args.bHideSelectionTip = true;
+	Args.NotifyHook = InNotifyHook;
+
+	TSharedPtr<IDetailsView> TransformationsPropertiesPropagator = PropertyModule.CreateDetailView(Args);
+	TransformationsPropertiesPropagator->SetObject(SoundWaveToView);
+	return TransformationsPropertiesPropagator;
+}
+
+void FTransformedWaveformViewFactory::SetUpWaveformPanelInteractions(TSharedRef<STransformedWaveformViewPanel> WaveformPanel, TSharedRef<FSparseSampledSequenceTransportCoordinator> TransportCoordinator, TSharedPtr<FWaveformEditorZoomController> InZoomController /*= nullptr*/, TSharedPtr<SWaveformTransformationsOverlay> InTransformationsOverlay /*=nullptr*/)
+{
+	FPointerEventHandler HandlePlayheadOverlayMouseButtonUp = FPointerEventHandler::CreateLambda([TransportCoordinator](const FGeometry& Geometry, const FPointerEvent& MouseEvent)
+	{
+		const bool HandleLeftMouseButton = MouseEvent.GetEffectingButton() == EKeys::LeftMouseButton;
+		if (HandleLeftMouseButton)
+		{
+			const float LocalWidth = Geometry.GetLocalSize().X;
+
+			if (LocalWidth > 0.f)
+			{
+
+				const float NewPosition = Geometry.AbsoluteToLocal(MouseEvent.GetScreenSpacePosition()).X / LocalWidth;
+				TransportCoordinator->ScrubFocusPoint(NewPosition, false);
+			}
+		}
+
+		return FReply::Handled();
+	});
+
+	FPointerEventHandler HandleTimeRulerMouseButtonUp = FPointerEventHandler::CreateLambda([WaveformPanel, TransportCoordinator, this](const FGeometry& Geometry, const FPointerEvent& MouseEvent)
+	{
+		TSharedRef<SWidget> WavePanelChildSlot = WaveformPanel->GetChildren()->GetChildAt(0);
+
+		for (int32 ChildIndex = 0; ChildIndex < WavePanelChildSlot->GetChildren()->Num(); ++ChildIndex)
+		{
+			TSharedRef<SWidget> RulerWidget = WavePanelChildSlot->GetChildren()->GetChildAt(ChildIndex);
+
+			if (RulerWidget->GetType() == TimeRulerWidgetName)
+			{
+				return HandleTimeRulerInteraction(TransformedWaveformViewFactory::EReceivedInteractionType::MouseButtonUp, TransportCoordinator, RulerWidget, MouseEvent, Geometry);
+			}
+		}
+
+		return FReply::Unhandled();
+	});
+
+	FPointerEventHandler HandleTimeRulerMouseButtonDown = FPointerEventHandler::CreateLambda([WaveformPanel, TransportCoordinator, this](const FGeometry& Geometry, const FPointerEvent& MouseEvent)
+	{
+		TSharedRef<SWidget> WavePanelChildSlot = WaveformPanel->GetChildren()->GetChildAt(0);
+
+		for (int32 ChildIndex = 0; ChildIndex < WavePanelChildSlot->GetChildren()->Num(); ++ChildIndex)
+		{
+			TSharedRef<SWidget> RulerWidget = WavePanelChildSlot->GetChildren()->GetChildAt(ChildIndex);
+			if (RulerWidget->GetType() == TimeRulerWidgetName)
+			{
+				return HandleTimeRulerInteraction(TransformedWaveformViewFactory::EReceivedInteractionType::MouseButtonDown, TransportCoordinator, RulerWidget, MouseEvent, Geometry);
+			}
+		}
+
+		return FReply::Unhandled();
+	});
+
+	FPointerEventHandler HandleTimeRulerMouseMove = FPointerEventHandler::CreateLambda([WaveformPanel, TransportCoordinator, this](const FGeometry& Geometry, const FPointerEvent& MouseEvent)
+	{
+		TSharedRef<SWidget> WavePanelChildSlot = WaveformPanel->GetChildren()->GetChildAt(0);
+
+		for (int32 ChildIndex = 0; ChildIndex < WavePanelChildSlot->GetChildren()->Num(); ++ChildIndex)
+		{
+			TSharedRef<SWidget> RulerWidget = WavePanelChildSlot->GetChildren()->GetChildAt(ChildIndex);
+			if (RulerWidget->GetType() == TimeRulerWidgetName)
+			{
+				return HandleTimeRulerInteraction(TransformedWaveformViewFactory::EReceivedInteractionType::MouseMove, TransportCoordinator, RulerWidget, MouseEvent, Geometry);
+			}
+		}
+
+		return FReply::Unhandled();
+
+	});
+
+	
+	FPointerEventHandler HandleMouseWheel = FPointerEventHandler::CreateLambda([InZoomController , InTransformationsOverlay](const FGeometry& Geometry, const FPointerEvent& MouseEvent)
+	{
+		if (InTransformationsOverlay)
+		{
+			FReply TransformationsOverlayReply = InTransformationsOverlay->OnMouseWheel(Geometry, MouseEvent);
+
+			if (TransformationsOverlayReply.IsEventHandled())
+			{
+				return TransformationsOverlayReply;
+			}
+		}
+
+		if (InZoomController)
+		{
+			InZoomController->ZoomByDelta(MouseEvent.GetWheelDelta());
+			return FReply::Handled();
+		}
+
+		return FReply::Unhandled();
+
+	});
+
+	WaveformPanel->SetOnMouseWheel(HandleMouseWheel);
+
+
+	WaveformPanel->SetOnPlayheadOverlayMouseButtonUp(HandlePlayheadOverlayMouseButtonUp);
+	WaveformPanel->SetOnTimeRulerMouseButtonUp(HandleTimeRulerMouseButtonUp);
+	WaveformPanel->SetOnTimeRulerMouseButtonDown(HandleTimeRulerMouseButtonDown);
+	WaveformPanel->SetOnTimeRulerMouseMove(HandleTimeRulerMouseMove);
+}
+
+FReply FTransformedWaveformViewFactory::HandleTimeRulerInteraction(const TransformedWaveformViewFactory::EReceivedInteractionType MouseInteractionType, TSharedRef<FSparseSampledSequenceTransportCoordinator> TransportCoordinator, const TSharedRef<SWidget> TimeRulerWidget, const FPointerEvent& MouseEvent, const FGeometry& Geometry)
+{
+
+	const float LocalWidth = Geometry.GetLocalSize().X;
+
+	if (LocalWidth > 0.f)
+	{
+		const FVector2D ScreenSpacePosition = MouseEvent.GetScreenSpacePosition();
+		const FVector2D CursorPosition = Geometry.AbsoluteToLocal(ScreenSpacePosition);
+		const float CursorXRatio = CursorPosition.X / LocalWidth;
+
+		switch (MouseInteractionType)
+		{
+		case TransformedWaveformViewFactory::EReceivedInteractionType::MouseButtonDown:
+			if (MouseEvent.IsMouseButtonDown(EKeys::LeftMouseButton))
+			{
+				return FReply::Handled().CaptureMouse(TimeRulerWidget).PreventThrottling();
+			}
+			break;
+		case TransformedWaveformViewFactory::EReceivedInteractionType::MouseMove:
+			if (MouseEvent.IsMouseButtonDown(EKeys::LeftMouseButton))
+			{
+				TransportCoordinator->ScrubFocusPoint(CursorXRatio, true);
+				return FReply::Handled().CaptureMouse(TimeRulerWidget);
+			}
+			break;
+		case TransformedWaveformViewFactory::EReceivedInteractionType::MouseButtonUp:
+			if (TimeRulerWidget->HasMouseCapture())
+			{
+				TransportCoordinator->ScrubFocusPoint(CursorXRatio, false);
+				return FReply::Handled().ReleaseMouseCapture();
+			}
+			break;
+		default:
+			static_assert(static_cast<int32>(TransformedWaveformViewFactory::EReceivedInteractionType::COUNT) == 3, "Possible missing switch case coverage for 'EReceivedInteractionType'");
+			break;
+		}
+	}
+
+	return FReply::Handled();
+}
+
+TUniquePtr<FTransformedWaveformViewFactory> FTransformedWaveformViewFactory::Instance;
diff --git a/Engine/Plugins/Editor/WaveformEditor/Source/WaveformEditor/Private/TransformedWaveformViewFactory.h b/Engine/Plugins/Editor/WaveformEditor/Source/WaveformEditor/Private/TransformedWaveformViewFactory.h
new file mode 100644
--- /dev/null
+++ b/Engine/Plugins/Editor/WaveformEditor/Source/WaveformEditor/Private/TransformedWaveformViewFactory.h
@@ -0,0 +1,47 @@
+// Copyright Epic Games, Inc. All Rights Reserved.
+
+#pragma once
+
+#include "WaveformEditorSequenceDataProvider.h"
+#include "Input/Reply.h"
+
+class FNotifyHook;
+class FSparseSampledSequenceTransportCoordinator;
+class FWaveformEditorZoomController; 
+class IDetailsView;
+class STransformedWaveformViewPanel;
+class SWaveformTransformationsOverlay;
+class SWidget;
+class USoundWave;
+struct FGeometry;
+struct FPointerEvent;
+struct FTransformedWaveformView;
+
+namespace TransformedWaveformViewFactory
+{
+	enum class EReceivedInteractionType
+	{
+		MouseButtonUp,
+		MouseButtonDown,
+		MouseMove,
+		COUNT
+	};
+}
+
+class WAVEFORMEDITOR_API FTransformedWaveformViewFactory
+{	
+public:
+	static FTransformedWaveformViewFactory& Get();
+	static void Create();
+
+	FTransformedWaveformView GetTransformedView(TObjectPtr<USoundWave> SoundWaveToView, TSharedRef<FSparseSampledSequenceTransportCoordinator> TransportCoordinator, FNotifyHook* OwnerNotifyHook = nullptr, TSharedPtr<FWaveformEditorZoomController> InZoomController = nullptr);
+
+private:
+	FOnTransformationsPropertiesRequired SetUpOwnerNotifyHook(TObjectPtr<USoundWave> SoundWaveToView, FNotifyHook* InNotifyHook);
+	TSharedPtr<IDetailsView> SetUpTransformationsPropertiesPropagator(TObjectPtr<USoundWave> SoundWaveToView, FNotifyHook* InNotifyHook);
+	void SetUpWaveformPanelInteractions(TSharedRef<STransformedWaveformViewPanel> WaveformPanel, TSharedRef<FSparseSampledSequenceTransportCoordinator> TransportCoordinator, TSharedPtr<FWaveformEditorZoomController> InZoomController = nullptr, TSharedPtr<SWaveformTransformationsOverlay> InTransformationsOverlay = nullptr);
+	FReply HandleTimeRulerInteraction(const TransformedWaveformViewFactory::EReceivedInteractionType MouseInteractionType, TSharedRef<FSparseSampledSequenceTransportCoordinator> TransportCoordinator, const TSharedRef<SWidget> TimeRulerWidget, const FPointerEvent& MouseEvent, const FGeometry& Geometry);
+	
+	static TUniquePtr<FTransformedWaveformViewFactory> Instance;
+
+};
\ No newline at end of file
diff --git a/Engine/Plugins/Editor/WaveformEditor/Source/WaveformEditor/Private/WaveformEditor.cpp b/Engine/Plugins/Editor/WaveformEditor/Source/WaveformEditor/Private/WaveformEditor.cpp
--- a/Engine/Plugins/Editor/WaveformEditor/Source/WaveformEditor/Private/WaveformEditor.cpp
+++ b/Engine/Plugins/Editor/WaveformEditor/Source/WaveformEditor/Private/WaveformEditor.cpp
@@ -6,21 +6,22 @@
 #include "Components/AudioComponent.h"
 #include "Misc/TransactionObjectEvent.h"
 #include "PropertyEditorModule.h"
-#include "SparseSampledSequenceTransportCoordinator.h"
 #include "Sound/SoundWave.h"
+#include "SparseSampledSequenceTransportCoordinator.h"
+#include "STransformedWaveformViewPanel.h"
 #include "Styling/AppStyle.h"
-#include "SWaveformPanel.h"
 #include "SWaveformTransformationsOverlay.h"
 #include "ToolMenus.h"
+#include "TransformedWaveformView.h"
+#include "TransformedWaveformViewFactory.h"
 #include "WaveformEditorCommands.h"
 #include "WaveformEditorDetailsCustomization.h"
 #include "WaveformEditorLog.h"
-#include "WaveformEditorRenderData.h"
+#include "WaveformEditorSequenceDataProvider.h"
 #include "WaveformEditorStyle.h"
 #include "WaveformEditorToolMenuContext.h"
 #include "WaveformEditorTransformationsSettings.h"
 #include "WaveformEditorWaveWriter.h"
-#include "WaveformTransformationsRenderManager.h"
 #include "Widgets/Docking/SDockTab.h"
 
 #define LOCTEXT_NAMESPACE "WaveformEditor"
@@ -42,11 +43,11 @@ bool FWaveformEditor::Init(const EToolkitMode::Type Mode, const TSharedPtr<ITool
 
 	bool bIsInitialized = true;
 	
-	bIsInitialized &= SetUpDetailsViews();
-	bIsInitialized &= SetUpWaveformPanel();
-	bIsInitialized &= SetupAudioComponent();
-	bIsInitialized &= SetUpTransportController();
-	bIsInitialized &= SetUpWaveWriter();
+	bIsInitialized &= CreateDetailsViews();
+	bIsInitialized &= CreateWaveformView();
+	bIsInitialized &= InitializeAudioComponent();
+	bIsInitialized &= CreateTransportController();
+	bIsInitialized &= CreateWaveWriter();
 	bIsInitialized &= BindDelegates();
 
 	bIsInitialized &= RegisterToolbar();
@@ -112,7 +113,7 @@ void FWaveformEditor::AddDefaultTransformations()
 	}
 }
 
-bool FWaveformEditor::SetupAudioComponent()
+bool FWaveformEditor::InitializeAudioComponent()
 {
 	if (!ensure(SoundWave))
 	{
@@ -146,7 +147,7 @@ bool FWaveformEditor::SetupAudioComponent()
 	return true;
 }
 
-bool FWaveformEditor::SetUpTransportController()
+bool FWaveformEditor::CreateTransportController()
 {
 	if (!ensure(AudioComponent))
 	{
@@ -158,7 +159,7 @@ bool FWaveformEditor::SetUpTransportController()
 	return TransportController != nullptr;
 }
 
-bool FWaveformEditor::SetUpZoom()
+bool FWaveformEditor::InitializeZoom()
 {
 	ZoomManager = MakeShared<FWaveformEditorZoomController>();
 	return ZoomManager != nullptr;
@@ -470,10 +471,10 @@ void FWaveformEditor::NotifyPostChange(const FPropertyChangedEvent& PropertyChan
 
 	if (bUpdateTransformationChain)
 	{
-		TransformationsRenderManager->GenerateLayersChain();
+		WaveformView.DataProvider->GenerateLayersChain();
 	}
 
-	TransformationsRenderManager->UpdateRenderElements();
+	WaveformView.DataProvider->UpdateRenderElements();
 	TransformationChainConfig = SoundWave->GetTransformationChainConfig();
 }
 
@@ -483,8 +484,8 @@ void FWaveformEditor::PostUndo(bool bSuccess)
 {
 	if (bSuccess)
 	{
-		TransformationsRenderManager->GenerateLayersChain();
-		TransformationsRenderManager->UpdateRenderElements();
+		WaveformView.DataProvider->GenerateLayersChain();
+		WaveformView.DataProvider->UpdateRenderElements();
 	}
 }
 
@@ -519,7 +520,7 @@ void FWaveformEditor::InitToolMenuContext(FToolMenuContext& MenuContext)
 	MenuContext.AddObject(Context);
 }
 
-bool FWaveformEditor::SetUpDetailsViews()
+bool FWaveformEditor::CreateDetailsViews()
 {
 	if (!ensure(SoundWave)) 
 	{
@@ -544,9 +545,6 @@ bool FWaveformEditor::SetUpDetailsViews()
 	TransformationsDetails->RegisterInstancedCustomPropertyLayout(SoundWave->GetClass(), TransformationsDetailsCustomizationInstance);
 	TransformationsDetails->SetObject(SoundWave);
 
-	TransformationsPropertiesPropagator = PropertyModule.CreateDetailView(Args);
-	TransformationsPropertiesPropagator->SetObject(SoundWave);
-	
 	return true;
 }
 
@@ -557,7 +555,7 @@ TSharedRef<SDockTab> FWaveformEditor::SpawnTab_WaveformDisplay(const FSpawnTabAr
 	return SNew(SDockTab)
 		.Label(LOCTEXT("WaveformDisplayTitle", "Waveform Display"))
 		[
-			WaveformPanel.ToSharedRef()
+			WaveformView.ViewWidget.ToSharedRef()
 		];
 }
 
@@ -613,7 +611,8 @@ TSharedRef<SDockTab> FWaveformEditor::SpawnTab_Transformations(const FSpawnTabAr
 			TransformationsDetails.ToSharedRef()
 		];
 }
-bool FWaveformEditor::SetUpWaveformPanel()
+
+bool FWaveformEditor::CreateWaveformView()
 {
 	if (!ensure(SoundWave))
 	{
@@ -621,52 +620,26 @@ bool FWaveformEditor::SetUpWaveformPanel()
 		return false;
 	}
 
-	if (!SetUpZoom())
+	if (!InitializeZoom())
 	{
 		UE_LOG(LogWaveformEditor, Warning, TEXT("Failed to Init Zoom Manager while setting up waveform panel"));
 		return false;
 	}
 
-	WaveformRenderData = MakeShared<FWaveformEditorRenderData>();
-	
 	TransportCoordinator = MakeShared<FSparseSampledSequenceTransportCoordinator>();
-	WaveformRenderData->OnRenderDataUpdated.AddSP(this, &FWaveformEditor::HandleRenderDataUpdate);
-	
-	FOnTransformationsPropertiesRequired OnTransformationPropertiesRequired = FOnTransformationsPropertiesRequired::CreateLambda([this](FTransformationsToPropertiesArray& InObjToPropsMap)
-	{
-		TSharedPtr<FWaveformTransformationsDetailsProvider> TransformationsDetailsProvider = MakeShared<FWaveformTransformationsDetailsProvider>();
-
-		FOnGetDetailCustomizationInstance ProviderInstance = FOnGetDetailCustomizationInstance::CreateLambda( [&]() 
-			{
-				return TransformationsDetailsProvider.ToSharedRef();
-			}
-		);
 
-		TransformationsPropertiesPropagator->RegisterInstancedCustomPropertyLayout(SoundWave->GetClass(), ProviderInstance);
-		TransformationsPropertiesPropagator->ForceRefresh();
+	WaveformView = FTransformedWaveformViewFactory::Get().GetTransformedView(SoundWave, TransportCoordinator.ToSharedRef(), this, ZoomManager);
 
-		for (FTransformationToPropertiesPair& ObjToPropsPair : InObjToPropsMap)
-		{
-			TransformationsDetailsProvider->GetHandlesForUObjectProperties(ObjToPropsPair.Key, ObjToPropsPair.Value);
-		}
-
-		TransformationsDetailsProvider.Reset();
-	});
-
-	TransformationsRenderManager = MakeShared<FWaveformTransformationsRenderManager>(SoundWave, OnTransformationPropertiesRequired);
-	TransformationsRenderManager->OnRenderDataGenerated.AddSP(WaveformRenderData.Get(), &FWaveformEditorRenderData::UpdateRenderData);
-
-	TSharedPtr<SWaveformTransformationsOverlay> TransformationsOverlay = SNew(SWaveformTransformationsOverlay, TransformationsRenderManager->GetTransformLayers(), TransportCoordinator.ToSharedRef());
-
-	TransformationsRenderManager->OnLayersChainGenerated.AddSP(TransformationsOverlay.Get(), &SWaveformTransformationsOverlay::OnLayerChainGenerated);
-	TransformationsRenderManager->OnRenderElementsUpdated.AddSP(TransformationsOverlay.Get(), &SWaveformTransformationsOverlay::UpdateLayerConstraints);
-	TransportCoordinator->OnDisplayRangeUpdated.AddSP(TransformationsOverlay.Get(), &SWaveformTransformationsOverlay::OnNewWaveformDisplayRange);
-	TransformationsRenderManager->UpdateRenderElements();
+	WaveformView.DataProvider->OnRenderElementsUpdated.AddSP(this, &FWaveformEditor::HandleRenderDataUpdate);
+	
+	ZoomManager->OnZoomRatioChanged.AddSP(TransportCoordinator.ToSharedRef(), &FSparseSampledSequenceTransportCoordinator::SetZoomRatio);
 
-	WaveformPanel = SNew(SWaveformPanel, WaveformRenderData.ToSharedRef(), TransportCoordinator.ToSharedRef(), ZoomManager.ToSharedRef(), TransformationsOverlay);
+	TransportCoordinator->OnDisplayRangeUpdated.AddSP(this, &FWaveformEditor::HandleDisplayRangeUpdate);
+	TransportCoordinator->OnFocusPointMoved.AddSP(WaveformView.ViewWidget.Get(), &STransformedWaveformViewPanel::SetPlayheadRatio);
 
-	return WaveformPanel != nullptr;
+	return WaveformView.ViewWidget != nullptr && WaveformView.DataProvider != nullptr;
 }
+
 void FWaveformEditor::HandlePlaybackPercentageChange(const UAudioComponent* InComponent, const USoundWave* InSoundWave, const float InPlaybackPercentage)
 {
 	const bool bIsStopped = AudioComponent->GetPlayState() == EAudioComponentPlayState::Stopped;
@@ -710,14 +683,18 @@ void FWaveformEditor::HandleAudioComponentPlayStateChanged(const UAudioComponent
 
 void FWaveformEditor::HandleRenderDataUpdate()
 {
-	check(WaveformRenderData != nullptr)
-
 	if (TransportCoordinator != nullptr)
 	{
-		TransportCoordinator->UpdatePlaybackRange(WaveformRenderData->GetTransformedWaveformBounds());
+		TransportCoordinator->UpdatePlaybackRange(WaveformView.DataProvider->GetTransformedWaveformBounds());
+		WaveformView.DataProvider->RequestSequenceView(TRange<double>::Inclusive(TransportCoordinator->GetDisplayRange().GetLowerBoundValue(), TransportCoordinator->GetDisplayRange().GetUpperBoundValue()));
 	}
 }
 
+void FWaveformEditor::HandleDisplayRangeUpdate(const TRange<float> NewRange)
+{
+	WaveformView.DataProvider->RequestSequenceView(TRange<double>::Inclusive(NewRange.GetLowerBoundValue(), NewRange.GetUpperBoundValue()));
+}
+
 void FWaveformEditor::HandlePlayheadScrub(const float InTargetPlayBackRatio, const bool bIsMoving)
 {
 	if (bIsMoving)
@@ -767,7 +744,7 @@ bool FWaveformEditor::CanPressPlayButton() const
 	return TransportController->CanPlay() && (TransportController->IsPaused() || !TransportController->IsPlaying());
 }
 
-bool FWaveformEditor::SetUpWaveWriter()
+bool FWaveformEditor::CreateWaveWriter()
 {
 	if (!ensure(SoundWave))
 	{
@@ -792,4 +769,4 @@ const UWaveformEditorTransformationsSettings* FWaveformEditor::GetWaveformEditor
 	return WaveformEditorTransformationsSettings;
 }
 
-#undef LOCTEXT_NAMESPACE
+#undef LOCTEXT_NAMESPACE
\ No newline at end of file
diff --git a/Engine/Plugins/Editor/WaveformEditor/Source/WaveformEditor/Private/WaveformEditorModule.cpp b/Engine/Plugins/Editor/WaveformEditor/Source/WaveformEditor/Private/WaveformEditorModule.cpp
--- a/Engine/Plugins/Editor/WaveformEditor/Source/WaveformEditor/Private/WaveformEditorModule.cpp
+++ b/Engine/Plugins/Editor/WaveformEditor/Source/WaveformEditor/Private/WaveformEditorModule.cpp
@@ -4,6 +4,7 @@
 
 #include "Modules/ModuleManager.h"
 #include "PropertyEditorModule.h"
+#include "TransformedWaveformViewFactory.h"
 #include "WaveformEditorCommands.h"
 #include "WaveformEditorInstantiator.h"
 #include "WaveformEditorLog.h"
@@ -14,6 +15,7 @@ DEFINE_LOG_CATEGORY(LogWaveformEditor);
 void FWaveformEditorModule::StartupModule()
 {
 	FWaveformEditorCommands::Register();
+	FTransformedWaveformViewFactory::Create();
 
 	WaveformEditorInstantiator = MakeShared<FWaveformEditorInstantiator>();
 	RegisterContentBrowserExtensions(WaveformEditorInstantiator.Get());
diff --git a/Engine/Plugins/Editor/WaveformEditor/Source/WaveformEditor/Public/TransformedWaveformView.h b/Engine/Plugins/Editor/WaveformEditor/Source/WaveformEditor/Public/TransformedWaveformView.h
new file mode 100644
--- /dev/null
+++ b/Engine/Plugins/Editor/WaveformEditor/Source/WaveformEditor/Public/TransformedWaveformView.h
@@ -0,0 +1,12 @@
+// Copyright Epic Games, Inc. All Rights Reserved.
+
+#pragma once
+
+#include "STransformedWaveformViewPanel.h"
+#include "WaveformEditorSequenceDataProvider.h"
+
+struct FTransformedWaveformView
+{
+	TSharedPtr<STransformedWaveformViewPanel> ViewWidget;
+	TSharedPtr<FWaveformEditorSequenceDataProvider> DataProvider;
+};
\ No newline at end of file
diff --git a/Engine/Plugins/Editor/WaveformEditor/Source/WaveformEditor/Public/WaveformEditor.h b/Engine/Plugins/Editor/WaveformEditor/Source/WaveformEditor/Public/WaveformEditor.h
--- a/Engine/Plugins/Editor/WaveformEditor/Source/WaveformEditor/Public/WaveformEditor.h
+++ b/Engine/Plugins/Editor/WaveformEditor/Source/WaveformEditor/Public/WaveformEditor.h
@@ -8,6 +8,7 @@
 #include "Toolkits/AssetEditorToolkit.h"
 #include "WaveformEditorTransportController.h"
 #include "WaveformEditorZoomController.h"
+#include "TransformedWaveformView.h"
 
 class IToolkitHost;
 class SDockTab;
@@ -20,9 +21,9 @@ class WAVEFORMEDITOR_API FWaveformEditor
 	, public FEditorUndoClient
 	, public FGCObject
 	, public FNotifyHook
-
 {
 public:
+
 	bool Init(const EToolkitMode::Type Mode, const TSharedPtr<IToolkitHost>& InitToolkitHost, USoundWave* SoundWaveToEdit);
 
 	/** FAssetEditorToolkit interface */
@@ -37,20 +38,17 @@ public:
 	/** FNotifyHook interface */
 	void NotifyPreChange(class FEditPropertyChain* PropertyAboutToChange) override {};
 	void NotifyPostChange(const FPropertyChangedEvent& PropertyChangedEvent, class FEditPropertyChain* PropertyThatChanged) override;
-	
 
 	/** FEditorUndo interface */
 	void PostUndo(bool bSuccess) override;
 	void PostRedo(bool bSuccess) override;
 	bool MatchesContext(const FTransactionContext& InContext, const TArray<TPair<UObject*, FTransactionObjectEvent>>& TransactionObjectContexts) const override;
-
 	virtual void InitToolMenuContext(FToolMenuContext& MenuContext) override;
 
 private:
-	bool SetupAudioComponent();
-	bool SetUpTransportController();
-	bool SetUpZoom();
-
+	bool InitializeAudioComponent();
+	bool CreateTransportController();
+	bool InitializeZoom();
 	bool BindDelegates();
 
 	/**	Sets the wave editor layout */
@@ -61,44 +59,34 @@ private:
 	bool BindCommands();
 	TSharedRef<SWidget> GenerateExportOptionsMenu();
 
-
 	/**	Details tabs set up */
 	TSharedRef<SDockTab> SpawnTab_Properties(const FSpawnTabArgs& Args);
 	TSharedRef<SDockTab> SpawnTab_Transformations(const FSpawnTabArgs& Args);
-
-	bool SetUpDetailsViews();
+	bool CreateDetailsViews();
 
 	/**	Waveform view tab setup */
 	TSharedRef<SDockTab> SpawnTab_WaveformDisplay(const FSpawnTabArgs& Args);
-	bool SetUpWaveformPanel();
+	bool CreateWaveformView();
 
 	/** Playback delegates handlers */
 	void HandlePlaybackPercentageChange(const UAudioComponent* InComponent, const USoundWave* InSoundWave, const float InPlaybackPercentage);
 	void HandleAudioComponentPlayStateChanged(const UAudioComponent* InAudioComponent, EAudioComponentPlayState NewPlayState);
 	void HandlePlayheadScrub(const float InTargetPlayBackRatio, const bool bIsMoving);
-	
+
+	/* Delegates Handlers*/
 	void HandleRenderDataUpdate();
+	void HandleDisplayRangeUpdate(const TRange<float>);
 
 	/** FGCObject overrides */
 	virtual void AddReferencedObjects(FReferenceCollector& Collector) override;
 	virtual FString GetReferencerName() const override;
-
 	bool CanPressPlayButton() const;
-
-	bool SetUpWaveWriter();
+	bool CreateWaveWriter();
 	void ExportWaveform();
-
 	const UWaveformEditorTransformationsSettings* GetWaveformEditorTransformationsSettings() const;
 	void AddDefaultTransformations();
 
-	/** Waveform Preview widget */
-	TSharedPtr<class SWaveformPanel> WaveformPanel;
-
-	/** Manages render information for waveform transforms */
-	TSharedPtr<class FWaveformTransformationsRenderManager> TransformationsRenderManager = nullptr;
-
-	/** Render Data of the displayed waveform */
-	TSharedPtr<class FWaveformEditorRenderData> WaveformRenderData = nullptr;
+	FTransformedWaveformView  WaveformView;
 
 	/** Exports the edited waveform to a new asset */
 	TSharedPtr<class FWaveformEditorWaveWriter> WaveWriter = nullptr;
@@ -118,9 +106,6 @@ private:
 	/** Transformations tab */
 	TSharedPtr<IDetailsView> TransformationsDetails;
 
-	/** Hidden details view to propagate property handles to transformation layers */
-	TSharedPtr<IDetailsView> TransformationsPropertiesPropagator;
-
 	/** Settings Editor App Identifier */
 	static const FName AppIdentifier;
 
@@ -128,21 +113,16 @@ private:
 	static const FName PropertiesTabId;
 	static const FName TransformationsTabId;
 	static const FName WaveformDisplayTabId;
-
 	static const FName EditorName;
 	static const FName ToolkitFName;
-
 	USoundWave* SoundWave = nullptr;
 	UAudioComponent* AudioComponent = nullptr;
-
 	bool bWasPlayingBeforeScrubbing = false;
 	bool bIsInteractingWithTransformations = false;
-
 	float LastReceivedPlaybackPercent = 0.f;
-
 	EAudioComponentPlayState TransformInteractionPlayState = EAudioComponentPlayState::Stopped;
 	float PlaybackTimeBeforeTransformInteraction = 0.f;
 	float StartTimeBeforeTransformInteraction = 0.f;
-
 	FWaveTransformUObjectConfiguration TransformationChainConfig;
 };
+
diff --git a/Engine/Plugins/Editor/WaveformEditor/Source/WaveformEditorWidgets/Private/STransformedWaveformViewPanel.cpp b/Engine/Plugins/Editor/WaveformEditor/Source/WaveformEditorWidgets/Private/STransformedWaveformViewPanel.cpp
new file mode 100644
--- /dev/null
+++ b/Engine/Plugins/Editor/WaveformEditor/Source/WaveformEditorWidgets/Private/STransformedWaveformViewPanel.cpp
@@ -0,0 +1,242 @@
+// Copyright Epic Games, Inc. All Rights Reserved.
+
+ #include "STransformedWaveformViewPanel.h"
+
+#include "SampledSequenceDisplayUnit.h"
+#include "SPlayheadOverlay.h"
+#include "SFixedSampledSequenceRuler.h"
+#include "SFixedSampledSequenceViewer.h"
+#include "SWaveformTransformationsOverlay.h"
+#include "SWaveformEditorInputRoutingOverlay.h"
+#include "WaveformEditorGridData.h"
+#include "WaveformEditorStyle.h"
+#include "Widgets/SBoxPanel.h"
+#include "Widgets/SCompoundWidget.h"
+#include "Widgets/SOverlay.h"
+
+void STransformedWaveformViewPanel::Construct(const FArguments& InArgs, const FFixedSampledSequenceView& InView)
+{
+	DisplayUnit = ESampledSequenceDisplayUnit::Seconds;
+	DataView = InView;
+
+	WaveformEditorStyle = &FWaveformEditorStyle::Get();
+	check(WaveformEditorStyle);
+
+	SetUpGridData();
+	SetUpWaveformViewer(GridData.ToSharedRef(), DataView);
+
+	if (InArgs._TransformationsOverlay)
+	{
+		WaveformTransformationsOverlay = InArgs._TransformationsOverlay;
+	}
+
+	SetupPlayheadOverlay();
+	SetUpInputRoutingOverlay();
+	SetUpTimeRuler(GridData.ToSharedRef());
+	SetUpInputOverrides(InArgs);
+	CreateLayout();
+}
+
+void STransformedWaveformViewPanel::CreateLayout()
+{
+	check(TimeRuler);
+	check(WaveformViewer);
+	check(InputRoutingOverlay);
+
+	TSharedPtr<SOverlay> WaveformView = SNew(SOverlay);
+	WaveformView->AddSlot()
+	[
+		WaveformViewer.ToSharedRef()
+	];
+
+	if (WaveformTransformationsOverlay)
+	{
+		WaveformView->AddSlot()
+		[
+			WaveformTransformationsOverlay.ToSharedRef()
+		];
+	}
+
+	WaveformView->AddSlot()
+	[
+		PlayheadOverlay.ToSharedRef()
+	];
+
+	WaveformView->AddSlot()
+	[
+		InputRoutingOverlay.ToSharedRef()
+	];
+
+
+	ChildSlot
+	[
+		SNew(SVerticalBox)
+		+ SVerticalBox::Slot().AutoHeight()
+		[
+			TimeRuler.ToSharedRef()
+		]
+		+ SVerticalBox::Slot()
+		[
+			WaveformView.ToSharedRef()
+		]
+	];
+}
+
+void STransformedWaveformViewPanel::SetUpTimeRuler(TSharedRef<FWaveformEditorGridData> InGridData)
+{
+	FFixedSampleSequenceRulerStyle* TimeRulerStyle = &WaveformEditorStyle->GetRegisteredWidgetStyle<FFixedSampleSequenceRulerStyle>("WaveformEditorRuler.Style").Get();
+	check(TimeRulerStyle);
+
+	TimeRuler = SNew(SFixedSampledSequenceRuler, InGridData).DisplayUnit(DisplayUnit).Style(TimeRulerStyle);
+	GridData->OnGridMetricsUpdated.AddSP(TimeRuler.ToSharedRef(), &SFixedSampledSequenceRuler::UpdateGridMetrics);
+	TimeRulerStyle->OnStyleUpdated.AddSP(TimeRuler.ToSharedRef(), &SFixedSampledSequenceRuler::OnStyleUpdated);
+	TimeRuler->OnTimeUnitMenuSelection.AddSP(this, &STransformedWaveformViewPanel::UpdateDisplayUnit);
+}
+
+void STransformedWaveformViewPanel::SetUpInputRoutingOverlay()
+{
+	FSampledSequenceViewerStyle* ViewerStyle = &WaveformEditorStyle->GetRegisteredWidgetStyle<FSampledSequenceViewerStyle>("WaveformViewer.Style").Get();
+	check(ViewerStyle);
+	
+	TArray<TSharedPtr<SWidget>> OverlaidWidgets;
+
+	if (WaveformTransformationsOverlay)
+	{
+		OverlaidWidgets.Add(WaveformTransformationsOverlay);
+	}
+
+	check(PlayheadOverlay)
+	OverlaidWidgets.Add(PlayheadOverlay);
+	InputRoutingOverlay = SNew(SWaveformEditorInputRoutingOverlay, OverlaidWidgets).Style(ViewerStyle);
+}
+
+void STransformedWaveformViewPanel::SetupPlayheadOverlay()
+{
+	FPlayheadOverlayStyle* PlayheadOverlayStyle = &WaveformEditorStyle->GetRegisteredWidgetStyle<FPlayheadOverlayStyle>("WaveformEditorPlayheadOverlay.Style").Get();
+	check(PlayheadOverlayStyle);
+
+	PlayheadOverlay = SNew(SPlayheadOverlay).Style(PlayheadOverlayStyle);
+	PlayheadOverlayStyle->OnStyleUpdated.AddSP(PlayheadOverlay.ToSharedRef(), &SPlayheadOverlay::OnStyleUpdated);
+}
+
+void STransformedWaveformViewPanel::SetUpWaveformViewer(TSharedRef<FWaveformEditorGridData> InGridData, const FFixedSampledSequenceView& InView)
+{
+	FSampledSequenceViewerStyle* WaveViewerStyle = &WaveformEditorStyle->GetRegisteredWidgetStyle<FSampledSequenceViewerStyle>("WaveformViewer.Style").Get();
+	check(WaveViewerStyle);
+
+	WaveformViewer = SNew(SFixedSampledSequenceViewer, InView.SampleData, InView.NumDimensions, InGridData).Style(WaveViewerStyle);
+	WaveViewerStyle->OnStyleUpdated.AddSP(WaveformViewer.ToSharedRef(), &SFixedSampledSequenceViewer::OnStyleUpdated);
+	GridData->OnGridMetricsUpdated.AddSP(WaveformViewer.ToSharedRef(), &SFixedSampledSequenceViewer::UpdateGridMetrics);
+}
+
+void STransformedWaveformViewPanel::SetUpGridData()
+{
+	FFixedSampleSequenceRulerStyle* RulerStyle = &WaveformEditorStyle->GetRegisteredWidgetStyle<FFixedSampleSequenceRulerStyle>("WaveformEditorRuler.Style").Get();
+	check(RulerStyle)
+	
+	GridData = MakeShared<FWaveformEditorGridData>(DataView.SampleData.Num() / DataView.NumDimensions, DataView.SampleRate, RulerStyle->DesiredWidth, &RulerStyle->TicksTextFont);
+}
+
+void STransformedWaveformViewPanel::ReceiveSequenceView(const FFixedSampledSequenceView InView, const uint32 FirstSampleIndex)
+{
+	DataView = InView;
+	if (GridData)
+	{
+		const uint32 FirstRenderedFrame = FirstSampleIndex / InView.NumDimensions;
+		const uint32 NumFrames = InView.SampleData.Num() / InView.NumDimensions;
+		GridData->UpdateDisplayRange(TRange<uint32>(FirstRenderedFrame, FirstRenderedFrame + NumFrames));
+	}
+
+	if (WaveformViewer)
+	{
+		WaveformViewer->UpdateView(InView.SampleData, InView.NumDimensions);
+	}
+}
+
+void STransformedWaveformViewPanel::SetPlayheadRatio(const float InRatio)
+{
+	CachedPlayheadRatio = InRatio;
+}
+
+void STransformedWaveformViewPanel::SetOnPlayheadOverlayMouseButtonUp(FPointerEventHandler InEventHandler)
+{
+	check(PlayheadOverlay)
+	PlayheadOverlay->SetOnMouseButtonUp(InEventHandler);
+}
+
+void STransformedWaveformViewPanel::SetOnTimeRulerMouseButtonUp(FPointerEventHandler InEventHandler)
+{
+	check(TimeRuler)
+	TimeRuler->SetOnMouseButtonUp(InEventHandler);
+}
+
+void STransformedWaveformViewPanel::SetOnTimeRulerMouseButtonDown(FPointerEventHandler InEventHandler)
+{
+	check(TimeRuler)
+	TimeRuler->SetOnMouseButtonDown(InEventHandler);
+}
+
+void STransformedWaveformViewPanel::SetOnTimeRulerMouseMove(FPointerEventHandler InEventHandler)
+{
+	check(TimeRuler)
+	TimeRuler->SetOnMouseMove(InEventHandler);
+}
+
+void STransformedWaveformViewPanel::SetOnMouseWheel(FPointerEventHandler InEventHandler)
+{
+	check(InputRoutingOverlay)
+	InputRoutingOverlay->OnMouseWheelDelegate = InEventHandler;
+}
+
+void STransformedWaveformViewPanel::UpdateDisplayUnit(const ESampledSequenceDisplayUnit InDisplayUnit)
+{
+	DisplayUnit = InDisplayUnit;
+	TimeRuler->UpdateDisplayUnit(DisplayUnit);
+}
+
+void STransformedWaveformViewPanel::Tick(const FGeometry& AllottedGeometry, const double InCurrentTime, const float InDeltaTime)
+{
+	const float PaintedWidth = AllottedGeometry.GetLocalSize().X;
+
+	if (PaintedWidth != CachedPixelWidth)
+	{
+		CachedPixelWidth = PaintedWidth;
+
+		if (GridData)
+		{
+			GridData->UpdateGridMetrics(PaintedWidth);
+		}
+	}
+
+	UpdatePlayheadPosition(PaintedWidth);
+
+}
+
+void STransformedWaveformViewPanel::UpdatePlayheadPosition(const float PaintedWidth)
+{
+
+	float PlayheadX = CachedPlayheadRatio * PaintedWidth;
+	PlayheadX = GridData ? GridData->SnapPositionToClosestFrame(PlayheadX) : PlayheadX;
+
+
+	if (PlayheadOverlay)
+	{
+		PlayheadOverlay->SetPlayheadPosition(PlayheadX);
+	}
+
+	if (TimeRuler)
+	{
+		TimeRuler->SetPlayheadPosition(PlayheadX);
+	}
+}
+
+void STransformedWaveformViewPanel::SetUpInputOverrides(const FArguments& InArgs)
+{
+	SetOnPlayheadOverlayMouseButtonUp(InArgs._OnPlayheadOverlayMouseButtonUp);
+
+	SetOnTimeRulerMouseButtonUp(InArgs._OnTimeRulerMouseButtonUp);
+	SetOnTimeRulerMouseButtonDown(InArgs._OnTimeRulerMouseButtonDown);
+	SetOnTimeRulerMouseMove(InArgs._OnTimeRulerMouseMove);
+	
+	SetOnMouseWheel(InArgs._OnMouseWheel);
+}
\ No newline at end of file
diff --git a/Engine/Plugins/Editor/WaveformEditor/Source/WaveformEditorWidgets/Private/SWaveformEditorInputRoutingOverlay.cpp b/Engine/Plugins/Editor/WaveformEditor/Source/WaveformEditorWidgets/Private/SWaveformEditorInputRoutingOverlay.cpp
--- a/Engine/Plugins/Editor/WaveformEditor/Source/WaveformEditorWidgets/Private/SWaveformEditorInputRoutingOverlay.cpp
+++ b/Engine/Plugins/Editor/WaveformEditor/Source/WaveformEditorWidgets/Private/SWaveformEditorInputRoutingOverlay.cpp
@@ -34,16 +34,12 @@ FReply SWaveformEditorInputRoutingOverlay::OnMouseMove(const FGeometry& MyGeomet
 
 FReply SWaveformEditorInputRoutingOverlay::OnMouseWheel(const FGeometry& MyGeometry, const FPointerEvent& MouseEvent)
 {
-	FReply InteractionReply = AudioWidgetsUtils::RouteMouseInput(&SWidget::OnMouseWheel, MouseEvent, MakeArrayView(OverlaidWidgets.GetData(), OverlaidWidgets.Num()));
-
-	if (!InteractionReply.IsEventHandled())
+	if (OnMouseWheelDelegate.IsBound())
 	{
-		OnNewMouseDelta.ExecuteIfBound(MouseEvent.GetWheelDelta());
-		InteractionReply = FReply::Handled();
+		return OnMouseWheelDelegate.Execute(MyGeometry, MouseEvent);
 	}
 
-	return InteractionReply;
-
+	return AudioWidgetsUtils::RouteMouseInput(&SWidget::OnMouseWheel, MouseEvent, MakeArrayView(OverlaidWidgets.GetData(), OverlaidWidgets.Num()));
 }
 
 FCursorReply SWaveformEditorInputRoutingOverlay::OnCursorQuery(const FGeometry& MyGeometry, const FPointerEvent& CursorEvent) const
diff --git a/Engine/Plugins/Editor/WaveformEditor/Source/WaveformEditorWidgets/Private/SWaveformEditorInputRoutingOverlay.h b/Engine/Plugins/Editor/WaveformEditor/Source/WaveformEditorWidgets/Private/SWaveformEditorInputRoutingOverlay.h
--- a/Engine/Plugins/Editor/WaveformEditor/Source/WaveformEditorWidgets/Private/SWaveformEditorInputRoutingOverlay.h
+++ b/Engine/Plugins/Editor/WaveformEditor/Source/WaveformEditorWidgets/Private/SWaveformEditorInputRoutingOverlay.h
@@ -7,8 +7,6 @@
 #include "Styling/SlateWidgetStyleAsset.h"
 #include "Widgets/SLeafWidget.h"
 
-DECLARE_DELEGATE_OneParam(FOnNewMouseDelta, const float /* new delta */)
-
 class FWaveformEditorGridData;
 class FSparseSampledSequenceTransportCoordinator;
 
@@ -31,7 +29,7 @@ public:
 	virtual FReply OnMouseMove(const FGeometry& MyGeometry, const FPointerEvent& MouseEvent) override;
 	virtual FCursorReply OnCursorQuery(const FGeometry& MyGeometry, const FPointerEvent& CursorEvent) const override;
 
-	FOnNewMouseDelta OnNewMouseDelta;\
+	FPointerEventHandler OnMouseWheelDelegate;
 
 private:
 	virtual int32 OnPaint(const FPaintArgs& Args, const FGeometry& AllottedGeometry, const FSlateRect& MyCullingRect, FSlateWindowElementList& OutDrawElements, int32 LayerId, const FWidgetStyle& InWidgetStyle, bool bParentEnabled) const override;
diff --git a/Engine/Plugins/Editor/WaveformEditor/Source/WaveformEditorWidgets/Private/SWaveformPanel.cpp b/Engine/Plugins/Editor/WaveformEditor/Source/WaveformEditorWidgets/Private/SWaveformPanel.cpp
deleted file mode 100644
diff --git a/Engine/Plugins/Editor/WaveformEditor/Source/WaveformEditorWidgets/Private/WaveformEditorRenderData.cpp b/Engine/Plugins/Editor/WaveformEditor/Source/WaveformEditorWidgets/Private/WaveformEditorRenderData.cpp
deleted file mode 100644
diff --git a/Engine/Plugins/Editor/WaveformEditor/Source/WaveformEditorWidgets/Private/WaveformEditorSequenceDataProvider.cpp b/Engine/Plugins/Editor/WaveformEditor/Source/WaveformEditorWidgets/Private/WaveformEditorSequenceDataProvider.cpp
new file mode 100644
--- /dev/null
+++ b/Engine/Plugins/Editor/WaveformEditor/Source/WaveformEditorWidgets/Private/WaveformEditorSequenceDataProvider.cpp
@@ -0,0 +1,250 @@
+// Copyright Epic Games, Inc. All Rights Reserved.
+
+#include "WaveformEditorSequenceDataProvider.h"
+
+#include "Audio.h"
+#include "DSP/BufferVectorOperations.h"
+#include "DSP/FloatArrayMath.h"
+#include "IWaveformTransformationRenderer.h"
+#include "Sound/SoundWave.h"
+#include "WaveformTransformationDurationRenderer.h"
+
+FWaveformEditorSequenceDataProvider::FWaveformEditorSequenceDataProvider(
+	TObjectPtr<USoundWave> InSoundWave, 
+	FOnTransformationsPropertiesRequired InTransformationChainPropertiesHandler)
+	: SoundWaveToRender(InSoundWave)
+	, LayersFactory(MakeUnique<FWaveformTransformationRenderLayerFactory>())
+	, ChainPropertiesHandler(InTransformationChainPropertiesHandler)
+{
+	GenerateLayersChain();	
+
+	if (!SoundWaveToRender->GetImportedSoundWaveData(ImportedRawPCMData, ImportedSampleRate, ImportedNumChannels))
+	{
+		UE_LOG(LogAudio, Warning, TEXT("Failed to get transformations render data for: %s"), *SoundWaveToRender->GetPathName());
+		return;
+	}
+}
+
+FFixedSampledSequenceView FWaveformEditorSequenceDataProvider::RequestSequenceView(const TRange<double> DataRatioRange)
+{
+	check(DataRatioRange.GetLowerBoundValue() >= 0)
+	check(DataRatioRange.GetUpperBoundValue() <= 1)
+	check(DataRatioRange.GetLowerBoundValue() != DataRatioRange.GetUpperBoundValue())
+
+	const int32 NumChannels = ImportedNumChannels;
+	const uint8 MinFramesToDisplay = 2;
+	const uint32 MinSamplesToDisplay = MinFramesToDisplay * NumChannels;
+	const uint32 NumOriginalSamples = TransformedPCMData.Num();
+	const uint32 NumOriginalFrames = NumOriginalSamples / NumChannels;
+
+	const uint32 FirstRenderedSample = FMath::Clamp(FMath::RoundToInt32(NumOriginalFrames * DataRatioRange.GetLowerBoundValue()), 0, NumOriginalFrames - MinFramesToDisplay) * NumChannels;
+	const uint32 NumFramesToRender = FMath::RoundToInt32(NumOriginalFrames * DataRatioRange.Size<double>());
+	const uint32 NumSamplesToRender = FMath::Clamp(NumFramesToRender * NumChannels, MinSamplesToDisplay, NumOriginalSamples - FirstRenderedSample);
+
+	check(NumSamplesToRender % NumChannels == 0 && FirstRenderedSample % NumChannels == 0);
+;
+	TArrayView<const float> SampleData = MakeArrayView(TransformedPCMData.GetData(), TransformedPCMData.Num()).Slice(FirstRenderedSample, NumSamplesToRender);
+
+	FFixedSampledSequenceView DataView{ SampleData , NumChannels, ImportedSampleRate};
+	OnDataViewGenerated.Broadcast(DataView, FirstRenderedSample);
+
+	return DataView;
+}
+
+void FWaveformEditorSequenceDataProvider::GenerateLayersChain()
+{
+	check(SoundWaveToRender);
+		
+	TransformationsToRender.Empty();
+	RenderLayers.Empty();
+
+	if (SoundWaveToRender->Transformations.Num() > 0)
+	{
+		FTransformationsToPropertiesArray TransformPropertiesMap;
+
+		for (const TObjectPtr<UWaveformTransformationBase> Transformation : SoundWaveToRender->Transformations)
+		{
+			if (Transformation != nullptr)
+			{
+				TArray<TSharedRef<IPropertyHandle>> TransformationPropertiesArray;
+				TransformPropertiesMap.Emplace(Transformation, TransformationPropertiesArray);
+			}
+		}
+
+		ChainPropertiesHandler.Execute(TransformPropertiesMap);
+		
+
+		for (FTransformationToPropertiesPair& ObjPropsPair : TransformPropertiesMap)
+		{
+			TransformationsToRender.Add(ObjPropsPair.Key);
+			TSharedPtr<IWaveformTransformationRenderer> TransformationUI = LayersFactory->Create(ObjPropsPair.Key, ObjPropsPair.Value);
+			FTransformationRenderLayerInfo RenderLayerInfo = FTransformationRenderLayerInfo(TransformationUI, FTransformationLayerConstraints(0.f, 1.f));
+			RenderLayers.Add(RenderLayerInfo);
+		}
+
+		CreateDurationHighlightLayer();
+	}
+
+	OnLayersChainGenerated.Broadcast(RenderLayers.GetData(), RenderLayers.Num());
+}
+
+void FWaveformEditorSequenceDataProvider::UpdateRenderElements()
+{
+	GenerateSequenceDataInternal();
+	OnRenderElementsUpdated.Broadcast();
+}
+
+void FWaveformEditorSequenceDataProvider::GenerateSequenceDataInternal()
+{
+	check(SoundWaveToRender);
+	
+	uint32 NumWaveformSamples = ImportedRawPCMData.Num() * sizeof(uint8) / sizeof(int16);
+	NumOriginalWaveformFrames = NumWaveformSamples / ImportedNumChannels;
+	uint32 FirstEditedSample = 0;
+	uint32 LastEditedSample = NumWaveformSamples;
+
+	if (TransformationsToRender.Num() > 0)
+	{
+		Audio::FWaveformTransformationWaveInfo TransformationInfo;
+
+		Audio::FAlignedFloatBuffer TransformationsBuffer;
+		Audio::FAlignedFloatBuffer OutputBuffer;
+
+		TransformationsBuffer.SetNumUninitialized(NumWaveformSamples);
+
+		Audio::ArrayPcm16ToFloat(MakeArrayView((int16*)ImportedRawPCMData.GetData(), NumWaveformSamples), TransformationsBuffer);
+		OutputBuffer = TransformationsBuffer;
+
+		TransformationInfo.Audio = &TransformationsBuffer;
+		TransformationInfo.NumChannels = ImportedNumChannels;
+		TransformationInfo.SampleRate = ImportedSampleRate;
+
+		TArray<Audio::FTransformationPtr> Transformations = CreateTransformations();
+
+		const bool bChainChangesFileLength = CanChainChangeFileLength(Transformations);
+
+		for (int32 i = 0; i < Transformations.Num(); ++i)
+		{
+			Transformations[i]->ProcessAudio(TransformationInfo);
+
+			if (RenderLayers[i].Key)
+			{
+				FWaveformTransformationRenderInfo RenderLayerInfo{ TransformationInfo.SampleRate, TransformationInfo.NumChannels, FirstEditedSample, LastEditedSample - FirstEditedSample };
+				RenderLayers[i].Key->SetTransformationWaveInfo(MoveTemp(RenderLayerInfo));
+
+				float LeftWidgetConstraint = FirstEditedSample / (float) NumWaveformSamples;
+				float RightWidgetConstraint = LastEditedSample / (float) NumWaveformSamples;
+
+				RenderLayers[i].Value = FTransformationLayerConstraints(LeftWidgetConstraint, RightWidgetConstraint);
+
+			}
+
+			FirstEditedSample += TransformationInfo.StartFrameOffset;
+			LastEditedSample = TransformationInfo.NumEditedSamples <= 0 ? LastEditedSample : FirstEditedSample + TransformationInfo.NumEditedSamples;
+
+			if (bChainChangesFileLength)
+			{
+				check(LastEditedSample > FirstEditedSample);
+				const int32 EditedBufferSize = LastEditedSample - FirstEditedSample;
+				FMemory::Memcpy(&OutputBuffer.GetData()[FirstEditedSample], TransformationsBuffer.GetData(), EditedBufferSize * sizeof(float));
+			}
+
+			TransformationInfo.StartFrameOffset = 0;
+			TransformationInfo.NumEditedSamples = 0;
+		}
+
+		check(DurationHiglightLayer)
+		FWaveformTransformationRenderInfo DurationLayerInfo{ TransformationInfo.SampleRate, TransformationInfo.NumChannels, FirstEditedSample, LastEditedSample - FirstEditedSample };
+		DurationHiglightLayer->SetTransformationWaveInfo(MoveTemp(DurationLayerInfo));
+		DurationHiglightLayer->SetOriginalWaveformFrames(NumOriginalWaveformFrames);
+		
+
+		UpdateTransformedWaveformBounds(FirstEditedSample, LastEditedSample, NumWaveformSamples);
+
+		if (!bChainChangesFileLength)
+		{
+			FMemory::Memcpy(TransformedPCMData.GetData(), TransformationsBuffer.GetData(), TransformationsBuffer.GetAllocatedSize());
+		}
+		else
+		{
+			if (TransformedPCMData.Num() != OutputBuffer.Num())
+			{
+				TransformedPCMData.SetNumUninitialized(OutputBuffer.Num());
+			}
+			FMemory::Memcpy(TransformedPCMData.GetData(), OutputBuffer.GetData(), OutputBuffer.GetAllocatedSize());
+		}
+		
+		const float MaxValue = Audio::ArrayMaxAbsValue(TransformedPCMData);
+
+		if (MaxValue > 1.f)
+		{
+			Audio::ArrayMultiplyByConstantInPlace(TransformedPCMData, 1.f / MaxValue);
+		}
+		
+		check(TransformationInfo.NumChannels > 0);
+		check(TransformationInfo.SampleRate > 0);
+	}
+	else
+	{
+		if (TransformedPCMData.Num() != NumWaveformSamples)
+		{
+			TransformedPCMData.SetNumUninitialized(NumWaveformSamples);
+
+		}
+		Audio::ArrayPcm16ToFloat(MakeArrayView((int16*)ImportedRawPCMData.GetData(), NumWaveformSamples), TransformedPCMData);
+	}
+}
+
+void FWaveformEditorSequenceDataProvider::UpdateTransformedWaveformBounds(const uint32 FirstEditedSample, const uint32 LastEditedSample, const float NumOriginalSamples)
+{
+	const float FirstEditedSampleOffset = FirstEditedSample / NumOriginalSamples;
+	const float LastEditedSampleOffset = (NumOriginalSamples - LastEditedSample) / NumOriginalSamples;
+
+	TransformedWaveformBounds.SetLowerBoundValue(FirstEditedSampleOffset);
+	TransformedWaveformBounds.SetUpperBoundValue(1 - LastEditedSampleOffset);
+}
+
+TArray<Audio::FTransformationPtr> FWaveformEditorSequenceDataProvider::CreateTransformations() const
+{
+	TArray<Audio::FTransformationPtr> TransformationPtrs;
+
+	for (const TObjectPtr<UWaveformTransformationBase> TransformationBase : TransformationsToRender)
+	{
+		TransformationPtrs.Add(TransformationBase->CreateTransformation());
+	}
+
+	return TransformationPtrs;
+}
+
+const bool FWaveformEditorSequenceDataProvider::CanChainChangeFileLength(const TArray<Audio::FTransformationPtr>& TransformationChain) const
+{
+	bool bCanChainChangeFileLength = false;
+
+	for (const Audio::FTransformationPtr& Transformation : TransformationChain)
+	{
+		bCanChainChangeFileLength |= Transformation->CanChangeFileLength();
+	}
+
+	return bCanChainChangeFileLength;
+}
+
+TArrayView<const FTransformationRenderLayerInfo> FWaveformEditorSequenceDataProvider::GetTransformLayers() const
+{
+	return MakeArrayView(RenderLayers.GetData(), RenderLayers.Num());
+}
+
+const TRange<float> FWaveformEditorSequenceDataProvider::GetTransformedWaveformBounds() const
+{
+	return TransformedWaveformBounds;
+}
+
+void FWaveformEditorSequenceDataProvider::CreateDurationHighlightLayer()
+{
+	if (!DurationHiglightLayer)
+	{
+		DurationHiglightLayer = MakeShared<FWaveformTransformationDurationRenderer>(NumOriginalWaveformFrames);
+	}
+
+	DurationHiglightLayer->SetTransformationWaveInfo(FWaveformTransformationRenderInfo());
+	RenderLayers.Emplace(DurationHiglightLayer, FTransformationLayerConstraints(0.f,1.f));
+}
\ No newline at end of file
diff --git a/Engine/Plugins/Editor/WaveformEditor/Source/WaveformEditorWidgets/Private/WaveformTransformationRendererMapper.cpp b/Engine/Plugins/Editor/WaveformEditor/Source/WaveformEditorWidgets/Private/WaveformTransformationRendererMapper.cpp
deleted file mode 100644
diff --git a/Engine/Plugins/Editor/WaveformEditor/Source/WaveformEditorWidgets/Public/IWaveformTransformationRenderer.h b/Engine/Plugins/Editor/WaveformEditor/Source/WaveformEditorWidgets/Public/IWaveformTransformationRenderer.h
deleted file mode 100644
diff --git a/Engine/Plugins/Editor/WaveformEditor/Source/WaveformEditorWidgets/Public/STransformedWaveformViewPanel.h b/Engine/Plugins/Editor/WaveformEditor/Source/WaveformEditorWidgets/Public/STransformedWaveformViewPanel.h
new file mode 100644
--- /dev/null
+++ b/Engine/Plugins/Editor/WaveformEditor/Source/WaveformEditorWidgets/Public/STransformedWaveformViewPanel.h
@@ -0,0 +1,84 @@
+// Copyright Epic Games, Inc. All Rights Reserved.
+
+#pragma once
+
+#include "HAL/Platform.h"
+#include "Widgets/SCompoundWidget.h"
+#include "IFixedSampledSequenceViewReceiver.h"
+
+enum class ESampledSequenceDisplayUnit;
+class SWaveformEditorInputRoutingOverlay;
+class SWaveformTransformationsOverlay;
+class SPlayheadOverlay;
+class SFixedSampledSequenceViewer;
+class SFixedSampledSequenceRuler;
+class FWaveformEditorZoomController;
+class FWaveformEditorStyle;
+class FWaveformEditorRenderData;
+class FWaveformEditorGridData;
+class FSparseSampledSequenceTransportCoordinator;
+
+class WAVEFORMEDITORWIDGETS_API STransformedWaveformViewPanel : public SCompoundWidget, public IFixedSampledSequenceViewReceiver 
+{
+public: 
+	SLATE_BEGIN_ARGS(STransformedWaveformViewPanel) {}
+		
+		SLATE_DEFAULT_SLOT(FArguments, InArgs)
+		
+		SLATE_ARGUMENT(TSharedPtr<SWaveformTransformationsOverlay>, TransformationsOverlay)
+
+		SLATE_ARGUMENT(FPointerEventHandler, OnPlayheadOverlayMouseButtonUp)
+
+		SLATE_ARGUMENT(FPointerEventHandler, OnTimeRulerMouseButtonUp)
+
+		SLATE_ARGUMENT(FPointerEventHandler, OnTimeRulerMouseButtonDown)
+
+		SLATE_ARGUMENT(FPointerEventHandler, OnTimeRulerMouseMove)
+
+		SLATE_ARGUMENT(FPointerEventHandler, OnMouseWheel)
+
+	SLATE_END_ARGS()
+	
+	void Construct(const FArguments& InArgs, const FFixedSampledSequenceView& InView);
+	virtual void Tick(const FGeometry& AllottedGeometry, const double InCurrentTime, const float InDeltaTime) override;
+	virtual void ReceiveSequenceView(const FFixedSampledSequenceView InView, const uint32 FirstSampleIndex = 0) override;
+	void SetPlayheadRatio(const float InRatio);
+
+	void SetOnPlayheadOverlayMouseButtonUp(FPointerEventHandler InEventHandler);
+	void SetOnTimeRulerMouseButtonUp(FPointerEventHandler InEventHandler);
+	void SetOnTimeRulerMouseButtonDown(FPointerEventHandler InEventHandler);
+	void SetOnTimeRulerMouseMove(FPointerEventHandler InEventHandler);
+	void SetOnMouseWheel(FPointerEventHandler InEventHandler);
+
+
+private:
+	void CreateLayout();
+
+	void SetUpGridData();
+	void SetupPlayheadOverlay();
+	void SetUpTimeRuler(TSharedRef<FWaveformEditorGridData> InGridData);
+	void SetUpWaveformViewer(TSharedRef<FWaveformEditorGridData> InGridData, const FFixedSampledSequenceView& InView);
+	void SetUpInputRoutingOverlay();
+	void SetUpInputOverrides(const FArguments& InArgs);
+
+	void UpdateDisplayUnit(const ESampledSequenceDisplayUnit InDisplayUnit);
+	void UpdatePlayheadPosition(const float PaintedWidth);
+
+	TSharedPtr<FWaveformEditorGridData> GridData;
+
+	TSharedPtr<SFixedSampledSequenceRuler> TimeRuler;
+	TSharedPtr<SFixedSampledSequenceViewer> WaveformViewer;
+ 	TSharedPtr<SWaveformTransformationsOverlay> WaveformTransformationsOverlay;
+	TSharedPtr<SWaveformEditorInputRoutingOverlay> InputRoutingOverlay;
+	TSharedPtr<SPlayheadOverlay> PlayheadOverlay;
+
+	float CachedPixelWidth = 0.f;
+
+	ESampledSequenceDisplayUnit DisplayUnit;
+
+	FWaveformEditorStyle* WaveformEditorStyle;
+
+	FFixedSampledSequenceView DataView;
+	
+	float CachedPlayheadRatio = 0.f;
+};
\ No newline at end of file
diff --git a/Engine/Plugins/Editor/WaveformEditor/Source/WaveformEditorWidgets/Public/SWaveformPanel.h b/Engine/Plugins/Editor/WaveformEditor/Source/WaveformEditorWidgets/Public/SWaveformPanel.h
deleted file mode 100644
diff --git a/Engine/Plugins/Editor/WaveformEditor/Source/WaveformEditorWidgets/Public/WaveformEditorRenderData.h b/Engine/Plugins/Editor/WaveformEditor/Source/WaveformEditorWidgets/Public/WaveformEditorRenderData.h
deleted file mode 100644
diff --git a/Engine/Plugins/Editor/WaveformEditor/Source/WaveformEditorWidgets/Public/WaveformEditorSequenceDataProvider.h b/Engine/Plugins/Editor/WaveformEditor/Source/WaveformEditorWidgets/Public/WaveformEditorSequenceDataProvider.h
new file mode 100644
--- /dev/null
+++ b/Engine/Plugins/Editor/WaveformEditor/Source/WaveformEditorWidgets/Public/WaveformEditorSequenceDataProvider.h
@@ -0,0 +1,115 @@
+// Copyright Epic Games, Inc. All Rights Reserved.
+
+#pragma once
+
+#include "Containers/ArrayView.h"
+#include "IFixedSampledSequenceViewProvider.h"
+#include "IWaveformTransformation.h"
+#include "WaveformTransformationRenderLayerFactory.h"
+
+class IWaveformTransformationRenderer;
+class FWaveformTransformationDurationRenderer;
+class USoundWave;
+class UWaveformTransformationBase;
+struct FFixedSampledSequenceViewView;
+
+using FTransformationLayerConstraints = TPair<float, float>;
+using FTransformationRenderLayerInfo = TPair<TSharedPtr<IWaveformTransformationRenderer>, FTransformationLayerConstraints>;
+using FTransformationToPropertiesPair = TPair<TObjectPtr<UWaveformTransformationBase>, TArray<TSharedRef<IPropertyHandle>>>;
+using FTransformationsToPropertiesArray = TArray<FTransformationToPropertiesPair>;
+
+DECLARE_DELEGATE_OneParam(FOnTransformationsPropertiesRequired, FTransformationsToPropertiesArray& /*Transformations Properties To Hand*/)
+DECLARE_MULTICAST_DELEGATE(FOnRenderElementsUpdated)
+DECLARE_MULTICAST_DELEGATE_TwoParams(FOnDataViewGenerated, FFixedSampledSequenceView, const uint32 /*FirstSampleIndex*/)
+DECLARE_MULTICAST_DELEGATE_TwoParams(FOnLayersChainGenerated, FTransformationRenderLayerInfo* /*First Layer ptr*/, const int32 /* NLayers */)
+
+
+
+// FWaveformEditorSequenceDataProvider													
+//																						
+// The Waveform Editor Sequence data provider can produce UI to display a chain of		
+// waveform transformations.															
+// 																						
+// The main UI elements created are:													
+//  TransformedPCMData : a float array containing the transformed samples of the		
+//						 waveform, that can be requested through the					
+//					     IFixedSampledSequenceViewProvider interface.					
+//  Transform Layers : an array of UI renderers for each transformation					
+//                     in the chain. Renderers should be contained in dummy widgets.	
+//																						
+// These are created so that the entire stack of transformation is displayable.			
+// E.g: if a 10s long file is trimmed from second 2, seconds 0 to 1 will 				
+// still be present in the render data.													
+// UI widgets are passed a struct with information about the transformation to display	
+// property (e.g. StartFrameOffset, SampleDuration, etc.).								
+// 																						
+// UIs for different transformations are registered with and spawned by					
+// FWaveformTransformationRenderLayerFactory.											
+// Transformations don't necessarily have a widget UI. In that case, they are only 		
+// reflected in the render data. 	
+
+
+class WAVEFORMEDITORWIDGETS_API FWaveformEditorSequenceDataProvider : public IFixedSampledSequenceViewProvider, public TSharedFromThis<FWaveformEditorSequenceDataProvider>
+{
+public:
+
+	FWaveformEditorSequenceDataProvider(
+		TObjectPtr<USoundWave> InSoundWave,
+		FOnTransformationsPropertiesRequired InTransformationChainPropertiesHandler
+	);
+
+	virtual FFixedSampledSequenceView RequestSequenceView(const TRange<double> DataRatioRange) override;
+
+	/** Used to generate the stack of transformations UI					*/
+	/** Should be called when the waveform transformation chain is changed	*/
+	void GenerateLayersChain();
+
+	/** Used to generate updated sequence data and pass transformation info to the widgets */
+	/** Should be called when the transformations parameters are changed				 */
+	void UpdateRenderElements();
+
+	TArrayView<const FTransformationRenderLayerInfo> GetTransformLayers() const;
+
+	const TRange<float> GetTransformedWaveformBounds() const;
+
+	/** Called when a new layer chain of transformations UI is created */
+	FOnLayersChainGenerated OnLayersChainGenerated;
+
+	/** Called when the different render elements are updated */
+	FOnRenderElementsUpdated OnRenderElementsUpdated;
+
+	/** Called when new data view is generated */
+	FOnDataViewGenerated OnDataViewGenerated;
+
+private:
+	void GenerateSequenceDataInternal();
+
+	void UpdateTransformedWaveformBounds(const uint32 FirstEditedSample, const uint32 LastEditedSample, const float NumOriginalSamples);
+
+	void CreateDurationHighlightLayer();
+
+	TArray<Audio::FTransformationPtr> CreateTransformations() const;
+	const bool CanChainChangeFileLength(const TArray<Audio::FTransformationPtr>& TransformationArray) const;
+
+	TArray<TObjectPtr<UWaveformTransformationBase>> TransformationsToRender;
+	TArray<FTransformationRenderLayerInfo> RenderLayers;
+
+	TObjectPtr<USoundWave> SoundWaveToRender = nullptr;
+
+	/** Imported Soundwave Data **/
+	TArray<uint8> ImportedRawPCMData;
+	uint16 ImportedNumChannels;
+	uint32 ImportedSampleRate;
+
+	TArray<float> TransformedPCMData;
+
+	TUniquePtr<FWaveformTransformationRenderLayerFactory> LayersFactory = nullptr;
+	TSharedPtr<FWaveformTransformationDurationRenderer> DurationHiglightLayer = nullptr;
+
+	FOnTransformationsPropertiesRequired ChainPropertiesHandler;
+
+	uint32 NumOriginalWaveformFrames = 0;
+
+	/* The bounds of the transformed waveform in relation to the original */
+	TRange<float> TransformedWaveformBounds = TRange<float>::Inclusive(0.f, 1.f);
+};
\ No newline at end of file
diff --git a/Engine/Plugins/Editor/WaveformEditor/Source/WaveformEditorWidgets/Public/WaveformTransformationRendererMapper.h b/Engine/Plugins/Editor/WaveformEditor/Source/WaveformEditorWidgets/Public/WaveformTransformationRendererMapper.h
deleted file mode 100644
diff --git a/Engine/Plugins/Editor/WaveformEditor/Source/WaveformEditorWidgets/WaveformEditorWidgets.build.cs b/Engine/Plugins/Editor/WaveformEditor/Source/WaveformEditorWidgets/WaveformEditorWidgets.build.cs
--- a/Engine/Plugins/Editor/WaveformEditor/Source/WaveformEditorWidgets/WaveformEditorWidgets.build.cs
+++ b/Engine/Plugins/Editor/WaveformEditor/Source/WaveformEditorWidgets/WaveformEditorWidgets.build.cs
@@ -22,10 +22,11 @@ namespace UnrealBuildTool.Rules
 					"EditorStyle",
 					"Engine",
 					"InputCore",
-					"Slate",
 					"SignalProcessing",
+					"Slate",
 					"SlateCore", 
-					"ToolMenus"
+					"ToolMenus",
+					"WaveformTransformationsWidgets",
 				}
 			);
 		}
diff --git a/Engine/Plugins/Editor/WaveformEditor/Source/WaveformTransformationsWidgets/Private/WaveformTransformationDurationRenderer.cpp b/Engine/Plugins/Editor/WaveformEditor/Source/WaveformTransformationsWidgets/Private/WaveformTransformationDurationRenderer.cpp
new file mode 100644
--- /dev/null
+++ b/Engine/Plugins/Editor/WaveformEditor/Source/WaveformTransformationsWidgets/Private/WaveformTransformationDurationRenderer.cpp
@@ -0,0 +1,53 @@
+// Copyright Epic Games, Inc. All Rights Reserved.
+
+#include "WaveformTransformationDurationRenderer.h"
+
+FWaveformTransformationDurationRenderer::FWaveformTransformationDurationRenderer(const uint32 InOriginalWaveformNumFrames)
+	: OriginalWaveformNumFrames(InOriginalWaveformNumFrames)
+{
+}
+
+int32 FWaveformTransformationDurationRenderer::OnPaint(const FPaintArgs& Args, const FGeometry& AllottedGeometry, const FSlateRect& MyCullingRect, FSlateWindowElementList& OutDrawElements, int32 LayerId, const FWidgetStyle& InWidgetStyle, bool bParentEnabled) const
+{
+	const float StartTimeRatio = (float)TransformationWaveInfo.StartFrameOffset / OriginalWaveformNumFrames / TransformationWaveInfo.NumChannels;
+	const uint32 EndSample = TransformationWaveInfo.StartFrameOffset + TransformationWaveInfo.NumAvilableSamples;
+	const float EndTimeRatio = (float)EndSample / OriginalWaveformNumFrames / TransformationWaveInfo.NumChannels;
+
+	const bool bRenderLeftBox = StartTimeRatio >= 0.f;
+	const bool bRenderRightBox = EndTimeRatio <= 1.f;
+
+	if (bRenderLeftBox)
+	{
+		const float RightMarginX = StartTimeRatio * AllottedGeometry.Size.X;
+
+		FSlateDrawElement::MakeBox(
+			OutDrawElements,
+			++LayerId,
+			AllottedGeometry.ToPaintGeometry(FVector2D(RightMarginX, AllottedGeometry.Size.Y), FSlateLayoutTransform()),
+			FAppStyle::GetBrush(WaveformTransformationDurationHiglightParams::BackgroundBrushName),
+			ESlateDrawEffect::None,
+			WaveformTransformationDurationHiglightParams::BoxColor);
+
+	}
+
+	if (bRenderRightBox)
+	{
+		const float LeftMarginX = EndTimeRatio * AllottedGeometry.Size.X;
+
+		FSlateDrawElement::MakeBox(
+			OutDrawElements,
+			++LayerId,
+			AllottedGeometry.ToPaintGeometry(FVector2D(AllottedGeometry.Size.X, AllottedGeometry.Size.Y), FSlateLayoutTransform(FVector2D(LeftMarginX, 0.f))),
+			FAppStyle::GetBrush(WaveformTransformationDurationHiglightParams::BackgroundBrushName),
+			ESlateDrawEffect::None,
+			WaveformTransformationDurationHiglightParams::BoxColor);
+
+	}
+
+	return LayerId;
+}
+
+void FWaveformTransformationDurationRenderer::SetOriginalWaveformFrames(const uint32 NumFrames)
+{
+	OriginalWaveformNumFrames = NumFrames;
+}
diff --git a/Engine/Plugins/Editor/WaveformEditor/Source/WaveformTransformationsWidgets/Private/WaveformTransformationDurationRenderer.h b/Engine/Plugins/Editor/WaveformEditor/Source/WaveformTransformationsWidgets/Private/WaveformTransformationDurationRenderer.h
deleted file mode 100644
diff --git a/Engine/Plugins/Editor/WaveformEditor/Source/WaveformTransformationsWidgets/Private/WaveformTransformationRendererBase.h b/Engine/Plugins/Editor/WaveformEditor/Source/WaveformTransformationsWidgets/Private/WaveformTransformationRendererBase.h
deleted file mode 100644
diff --git a/Engine/Plugins/Editor/WaveformEditor/Source/WaveformTransformationsWidgets/Private/WaveformTransformationRendererMapper.cpp b/Engine/Plugins/Editor/WaveformEditor/Source/WaveformTransformationsWidgets/Private/WaveformTransformationRendererMapper.cpp
new file mode 100644
--- /dev/null
+++ b/Engine/Plugins/Editor/WaveformEditor/Source/WaveformTransformationsWidgets/Private/WaveformTransformationRendererMapper.cpp
@@ -0,0 +1,31 @@
+// Copyright Epic Games, Inc. All Rights Reserved.
+
+#include "WaveformTransformationRendererMapper.h"
+
+FWaveformTransformationRendererMapper& FWaveformTransformationRendererMapper::Get()
+{
+	check(Instance);
+	return *Instance;
+}
+
+void FWaveformTransformationRendererMapper::Init()
+{
+	if (Instance == nullptr)
+	{
+		Instance = MakeUnique<FWaveformTransformationRendererMapper>();
+	}
+}
+
+void FWaveformTransformationRendererMapper::UnregisterRenderer(const UClass* SupportedTransformation)
+{
+	check(Instance)
+	Instance->TransformationsTypeMap.FindAndRemoveChecked(SupportedTransformation);
+}
+
+FWaveformTransformRendererInstantiator* FWaveformTransformationRendererMapper::GetRenderer(const UClass* WaveformTransformationClass)
+{
+	check(Instance)
+	return Instance->TransformationsTypeMap.Find(WaveformTransformationClass);
+}
+
+TUniquePtr<FWaveformTransformationRendererMapper> FWaveformTransformationRendererMapper::Instance;
diff --git a/Engine/Plugins/Editor/WaveformEditor/Source/WaveformTransformationsWidgets/Private/WaveformTransformationsRenderManager.cpp b/Engine/Plugins/Editor/WaveformEditor/Source/WaveformTransformationsWidgets/Private/WaveformTransformationsRenderManager.cpp
deleted file mode 100644
diff --git a/Engine/Plugins/Editor/WaveformEditor/Source/WaveformTransformationsWidgets/Public/IWaveformTransformationRenderer.h b/Engine/Plugins/Editor/WaveformEditor/Source/WaveformTransformationsWidgets/Public/IWaveformTransformationRenderer.h
new file mode 100644
--- /dev/null
+++ b/Engine/Plugins/Editor/WaveformEditor/Source/WaveformTransformationsWidgets/Public/IWaveformTransformationRenderer.h
@@ -0,0 +1,52 @@
+// Copyright Epic Games, Inc. All Rights Reserved.
+
+#pragma once
+
+#include "IWaveformTransformation.h"
+#include "Templates/Function.h"
+
+class IPropertyHandle;
+class FCursorReply;
+class FPaintArgs;
+class FReply;
+class FSlateRect;
+class FSlateWindowElementList;
+class FWidgetStyle;
+class SWidget;
+class UClass;
+struct FGeometry;
+struct FPointerEvent;
+
+struct FWaveformTransformationRenderInfo
+{
+	float SampleRate = 0.f;
+	int32 NumChannels = 0;
+	uint32 StartFrameOffset = 0;	//the first frame this transformation can act upon
+	uint32 NumAvilableSamples = 0;	//the number of samples available to the transformation
+};
+
+class IWaveformTransformationRenderer
+{
+public:
+	virtual ~IWaveformTransformationRenderer() {};
+
+	/* SWidget interceptor methods */
+	virtual int32 OnPaint(const FPaintArgs& Args, const FGeometry& AllottedGeometry, const FSlateRect& MyCullingRect, FSlateWindowElementList& OutDrawElements, int32 LayerId, const FWidgetStyle& InWidgetStyle, bool bParentEnabled) const = 0;
+	virtual void Tick(const FGeometry& AllottedGeometry, const double InCurrentTime, const float InDeltaTime) = 0;
+	virtual FReply OnMouseButtonDown(SWidget& OwnerWidget, const FGeometry& MyGeometry, const FPointerEvent& MouseEvent) = 0;
+	virtual FReply OnMouseButtonUp(SWidget& OwnerWidget, const FGeometry& MyGeometry, const FPointerEvent& MouseEvent) = 0;
+	virtual FReply OnMouseButtonDoubleClick(SWidget& OwnerWidget, const FGeometry& InMyGeometry, const FPointerEvent& InMouseEvent) = 0;
+	virtual FReply OnMouseMove(SWidget& OwnerWidget, const FGeometry& MyGeometry, const FPointerEvent& MouseEvent) = 0;
+	virtual FReply OnMouseWheel(SWidget& OwnerWidget, const FGeometry& MyGeometry, const FPointerEvent& MouseEvent) = 0;
+	virtual void OnMouseEnter(const FGeometry& MyGeometry, const FPointerEvent& MouseEvent) = 0;
+	virtual void OnMouseLeave(const FPointerEvent& MouseEvent) = 0;
+	virtual FCursorReply OnCursorQuery(const FGeometry& MyGeometry, const FPointerEvent& CursorEvent) const = 0;
+
+	/* Gets info about the transformation being rendered */
+	virtual void SetTransformationWaveInfo(const FWaveformTransformationRenderInfo& InWaveInfo) = 0;
+	
+	/* Gets the transformations properties. IPropertyHandle is used to get/set properties to have full support of change notifiers, editor undo ecc  */
+	virtual void SetPropertyHandles(const TArray<TSharedRef<IPropertyHandle>>& InPropertyHandles) = 0;
+};
+
+
diff --git a/Engine/Plugins/Editor/WaveformEditor/Source/WaveformTransformationsWidgets/Public/WaveformTransformationDurationRenderer.h b/Engine/Plugins/Editor/WaveformEditor/Source/WaveformTransformationsWidgets/Public/WaveformTransformationDurationRenderer.h
new file mode 100644
--- /dev/null
+++ b/Engine/Plugins/Editor/WaveformEditor/Source/WaveformTransformationsWidgets/Public/WaveformTransformationDurationRenderer.h
@@ -0,0 +1,27 @@
+// Copyright Epic Games, Inc. All Rights Reserved.
+
+#pragma once
+
+#include "Layout/Geometry.h"
+#include "Rendering/DrawElements.h"
+#include "Styling/AppStyle.h"
+#include "WaveformTransformationRendererBase.h"
+
+
+namespace WaveformTransformationDurationHiglightParams
+{
+	const FLazyName BackgroundBrushName = TEXT("WhiteBrush");
+	const FLinearColor BoxColor = FLinearColor(0.f, 0.f, 0.f, 0.7f);
+}
+
+class WAVEFORMTRANSFORMATIONSWIDGETS_API FWaveformTransformationDurationRenderer : public FWaveformTransformationRendererBase
+{
+public:
+	explicit FWaveformTransformationDurationRenderer(const uint32 InOriginalWaveformNumFrames);
+
+	virtual int32 OnPaint(const FPaintArgs& Args, const FGeometry& AllottedGeometry, const FSlateRect& MyCullingRect, FSlateWindowElementList& OutDrawElements, int32 LayerId, const FWidgetStyle& InWidgetStyle, bool bParentEnabled) const override;
+	void SetOriginalWaveformFrames(const uint32 NumFrames);
+
+private:
+	uint32 OriginalWaveformNumFrames = 0;
+};
\ No newline at end of file
diff --git a/Engine/Plugins/Editor/WaveformEditor/Source/WaveformTransformationsWidgets/Public/WaveformTransformationRendererBase.h b/Engine/Plugins/Editor/WaveformEditor/Source/WaveformTransformationsWidgets/Public/WaveformTransformationRendererBase.h
new file mode 100644
--- /dev/null
+++ b/Engine/Plugins/Editor/WaveformEditor/Source/WaveformTransformationsWidgets/Public/WaveformTransformationRendererBase.h
@@ -0,0 +1,55 @@
+// Copyright Epic Games, Inc. All Rights Reserved.
+
+#pragma once
+
+#include "IDetailCustomization.h"
+#include "IWaveformTransformationRenderer.h"
+#include "PropertyHandle.h"
+
+class WAVEFORMTRANSFORMATIONSWIDGETS_API FWaveformTransformationRendererBase : public IWaveformTransformationRenderer
+{
+public:
+	FWaveformTransformationRendererBase() = default;
+	
+	virtual int32 OnPaint(const FPaintArgs& Args, const FGeometry& AllottedGeometry, const FSlateRect& MyCullingRect, FSlateWindowElementList& OutDrawElements, int32 LayerId, const FWidgetStyle& InWidgetStyle, bool bParentEnabled) const override;
+	virtual void Tick(const FGeometry& AllottedGeometry, const double InCurrentTime, const float InDeltaTime) override {};
+	virtual FReply OnMouseButtonDown(SWidget& OwnerWidget, const FGeometry& MyGeometry, const FPointerEvent& MouseEvent) override;
+	virtual FReply OnMouseButtonUp(SWidget& OwnerWidget, const FGeometry& MyGeometry, const FPointerEvent& MouseEvent) override;
+	virtual FReply OnMouseButtonDoubleClick(SWidget& OwnerWidget, const FGeometry& InMyGeometry, const FPointerEvent& InMouseEvent) override;
+	virtual FReply OnMouseMove(SWidget& OwnerWidget, const FGeometry& MyGeometry, const FPointerEvent& MouseEvent) override;
+	virtual FReply OnMouseWheel(SWidget& OwnerWidget, const FGeometry& MyGeometry, const FPointerEvent& MouseEvent) override;
+	virtual void OnMouseEnter(const FGeometry& MyGeometry, const FPointerEvent& MouseEvent) override {};
+	virtual void OnMouseLeave(const FPointerEvent& MouseEvent) override {};
+	virtual FCursorReply OnCursorQuery(const FGeometry& MyGeometry, const FPointerEvent& CursorEvent) const override;
+
+	virtual void SetTransformationWaveInfo(const FWaveformTransformationRenderInfo& InWaveInfo) override;
+	virtual void SetPropertyHandles(const TArray<TSharedRef<IPropertyHandle>>& InPropertyHandles) override;
+
+protected:
+	TSharedPtr<IPropertyHandle> GetPropertyHandle(const FName& PropertyName) const;
+
+	template<typename T>
+	T GetPropertyValue(const TSharedPtr<IPropertyHandle>& PropertyHandle)
+	{
+		check(PropertyHandle.IsValid());
+		T OutValue;
+		PropertyHandle->GetValue(OutValue);
+		return OutValue;
+	}
+
+	template<typename T>
+	void SetPropertyValue(const FName& PropertyName, const T Value, const EPropertyValueSetFlags::Type Flags = EPropertyValueSetFlags::DefaultFlags)
+	{
+		TSharedPtr<IPropertyHandle> Handle = GetPropertyHandle(PropertyName);
+		Handle->SetValue(Value, Flags);
+	}
+
+
+	int32 BeginTransaction(const TCHAR* TransactionContext, const FText& Description, UObject* PrimaryObject);
+	int32 EndTransaction();
+
+	FWaveformTransformationRenderInfo TransformationWaveInfo;
+
+private:
+	TArray<TSharedRef<IPropertyHandle>> CachedPropertyHandles;	
+};
\ No newline at end of file
diff --git a/Engine/Plugins/Editor/WaveformEditor/Source/WaveformTransformationsWidgets/Public/WaveformTransformationRendererMapper.h b/Engine/Plugins/Editor/WaveformEditor/Source/WaveformTransformationsWidgets/Public/WaveformTransformationRendererMapper.h
new file mode 100644
--- /dev/null
+++ b/Engine/Plugins/Editor/WaveformEditor/Source/WaveformTransformationsWidgets/Public/WaveformTransformationRendererMapper.h
@@ -0,0 +1,37 @@
+// Copyright Epic Games, Inc. All Rights Reserved.
+
+#pragma once
+
+#include "IWaveformTransformationRenderer.h"
+
+using FWaveformTransformRendererInstantiator = TFunction<TSharedPtr<IWaveformTransformationRenderer>()>;
+
+class WAVEFORMTRANSFORMATIONSWIDGETS_API FWaveformTransformationRendererMapper
+{
+public:
+	/** Access the singleton instance for mapper */
+	static FWaveformTransformationRendererMapper& Get();
+	static void Init();
+
+	template<typename ConcreteRendererType>
+	bool RegisterRenderer(const UClass* SupportedTransformation)
+	{
+		FWaveformTransformRendererInstantiator Instantiator = []()
+		{
+			return MakeShared<ConcreteRendererType>();
+		};
+
+		check(Instance)
+		Instance->TransformationsTypeMap.Add(SupportedTransformation, Instantiator);
+		return true;
+	}
+
+	void UnregisterRenderer(const UClass* SupportedTransformation);
+
+	FWaveformTransformRendererInstantiator* GetRenderer(const UClass* WaveformTransformationClass);
+
+private:
+	static TUniquePtr<FWaveformTransformationRendererMapper> Instance;
+
+	TMap<const UClass*, FWaveformTransformRendererInstantiator> TransformationsTypeMap;
+};
\ No newline at end of file
diff --git a/Engine/Plugins/Editor/WaveformEditor/Source/WaveformTransformationsWidgets/Public/WaveformTransformationsRenderManager.h b/Engine/Plugins/Editor/WaveformEditor/Source/WaveformTransformationsWidgets/Public/WaveformTransformationsRenderManager.h
deleted file mode 100644
diff --git a/Engine/Plugins/Editor/WaveformEditor/Source/WaveformTransformationsWidgets/WaveformTransformationsWidgets.build.cs b/Engine/Plugins/Editor/WaveformEditor/Source/WaveformTransformationsWidgets/WaveformTransformationsWidgets.build.cs
--- a/Engine/Plugins/Editor/WaveformEditor/Source/WaveformTransformationsWidgets/WaveformTransformationsWidgets.build.cs
+++ b/Engine/Plugins/Editor/WaveformEditor/Source/WaveformTransformationsWidgets/WaveformTransformationsWidgets.build.cs
@@ -27,7 +27,6 @@ namespace UnrealBuildTool.Rules
 					"ToolMenus", 
 					"UnrealEd",
 					"WaveformTransformations",
-					"WaveformEditorWidgets",
 				}
 			);
 		}
diff --git a/Engine/Plugins/Runtime/AudioWidgets/Source/AudioWidgets/Private/FixedSampledSequenceGridData.cpp b/Engine/Plugins/Runtime/AudioWidgets/Source/AudioWidgets/Private/FixedSampledSequenceGridData.cpp
new file mode 100644
--- /dev/null
+++ b/Engine/Plugins/Runtime/AudioWidgets/Source/AudioWidgets/Private/FixedSampledSequenceGridData.cpp
@@ -0,0 +1,97 @@
+// Copyright Epic Games, Inc. All Rights Reserved.
+
+#include "FixedSampledSequenceGridData.h"
+
+#include "Fonts/FontMeasure.h"
+#include "Framework/Application/SlateApplication.h"
+#include "GenericPlatform/GenericPlatformMath.h"
+
+FFixedSampledSequenceGridData::FFixedSampledSequenceGridData(const uint32 InTotalFrames, const uint32 InSampleRateHz, const float InGridSizePixels /*= 1.f*/, const FSlateFontInfo* InTicksTimeFont /*= nullptr*/)
+	: TotalFrames(InTotalFrames)
+	, DisplayRange(TRange<uint32>::Inclusive(0, InTotalFrames))
+	, GridSizePixels(InGridSizePixels)
+	, TicksTimeFont(InTicksTimeFont)
+	, GridFrameRate(InSampleRateHz, 1)
+{
+	UpdateGridMetrics(GridSizePixels);
+}
+
+void FFixedSampledSequenceGridData::UpdateDisplayRange(const TRange<uint32> InDisplayRange)
+{
+	check(DisplayRange.Size<uint32>() >= 2)
+	DisplayRange = InDisplayRange;
+	UpdateGridMetrics(GridSizePixels);
+}
+
+bool FFixedSampledSequenceGridData::UpdateGridMetrics(const float InGridSizePixels)
+{
+	if (GridSizePixels != InGridSizePixels)
+	{
+		GridSizePixels = InGridSizePixels;
+	}
+
+	const float WaveformDurationSeconds = TotalFrames / (double)GridFrameRate.Numerator;
+	const double StartTimeSeconds = DisplayRange.GetLowerBoundValue() / (double)GridFrameRate.Numerator;
+	const double DisplayedDuration = (DisplayRange.Size<double>() - 1) / GridFrameRate.Numerator; //we account for one less frame so we can distribute them evenly on the screen
+	const double PixelsPerSecond = GridSizePixels / DisplayedDuration;
+	double MajorGridStepSeconds = 0.0;
+	float MinTicksPixelSpacing = 30.0f;
+
+	if (TicksTimeFont)
+	{
+		TSharedRef<FSlateFontMeasure> FontMeasureService = FSlateApplication::Get().GetRenderer()->GetFontMeasureService();
+		FString LongestDisplayedTimeString = FString::Printf(TEXT("%.3f"), WaveformDurationSeconds * DisplayRange.GetUpperBoundValue());
+		FVector2D MaxTextSize = FontMeasureService->Measure(LongestDisplayedTimeString, *TicksTimeFont);
+		MinTicksPixelSpacing = MaxTextSize.X;
+	}
+
+
+	if (!GridFrameRate.ComputeGridSpacing(PixelsPerSecond, MajorGridStepSeconds, GridMetrics.NumMinorGridDivisions, MinTicksPixelSpacing + 5.f))
+	{
+		return false;
+	}
+
+	const bool bZeroGridOffset = DisplayRange.Size<uint32>() < GridSizePixels;
+	const double GridOffset = bZeroGridOffset ? 0 : FGenericPlatformMath::Fmod(StartTimeSeconds, MajorGridStepSeconds);
+	GridMetrics.FirstMajorTickX = (0.f - GridOffset) * PixelsPerSecond;
+	GridMetrics.MajorGridXStep = MajorGridStepSeconds * PixelsPerSecond;
+	GridMetrics.StartFrame = DisplayRange.GetLowerBoundValue();
+	GridMetrics.PixelsPerFrame = PixelsPerSecond / GridFrameRate.Numerator;
+	GridMetrics.SampleRate = GridFrameRate.Numerator;
+
+	if (OnGridMetricsUpdated.IsBound())
+	{
+		OnGridMetricsUpdated.Broadcast();
+	}
+
+	return true;
+}
+
+const FFixedSampledSequenceGridMetrics FFixedSampledSequenceGridData::GetGridMetrics() const
+{
+	return GridMetrics;
+}
+
+void FFixedSampledSequenceGridData::SetTicksTimeFont(const FSlateFontInfo* InNewFont)
+{
+	TicksTimeFont = InNewFont;
+}
+
+const float FFixedSampledSequenceGridData::SnapPositionToClosestFrame(const float InPixelPosition) const
+{
+	const float DistanceFromPreviousSample = FGenericPlatformMath::Fmod(InPixelPosition, GridMetrics.PixelsPerFrame);
+	const bool bSnapToNext = DistanceFromPreviousSample > GridMetrics.PixelsPerFrame / 2.f;
+	float SnappedPosition = 0.f;
+
+	if (bSnapToNext)
+	{
+		SnappedPosition = InPixelPosition + (GridMetrics.PixelsPerFrame - DistanceFromPreviousSample);
+	}
+	else
+	{
+		SnappedPosition = InPixelPosition - DistanceFromPreviousSample;
+	}
+
+
+	return SnappedPosition;
+}
\ No newline at end of file
diff --git a/Engine/Plugins/Runtime/AudioWidgets/Source/AudioWidgets/Private/SparseSampledSequenceTransportCoordinator.cpp b/Engine/Plugins/Runtime/AudioWidgets/Source/AudioWidgets/Private/SparseSampledSequenceTransportCoordinator.cpp
--- a/Engine/Plugins/Runtime/AudioWidgets/Source/AudioWidgets/Private/SparseSampledSequenceTransportCoordinator.cpp
+++ b/Engine/Plugins/Runtime/AudioWidgets/Source/AudioWidgets/Private/SparseSampledSequenceTransportCoordinator.cpp
@@ -53,6 +53,7 @@ const TRange<float> FSparseSampledSequenceTransportCoordinator::GetDisplayRange(
 void FSparseSampledSequenceTransportCoordinator::MoveFocusPoint(const float InFocusPoint)
 {
 	FocusPoint = InFocusPoint;
+	OnFocusPointMoved.Broadcast(FocusPoint);
 }
 
 void FSparseSampledSequenceTransportCoordinator::UpdateZoomRatioAndDisplayRange(const float NewZoomRatio)
@@ -207,8 +208,8 @@ void FSparseSampledSequenceTransportCoordinator::UpdateDisplayRange(const float
 {
 	check(MinValue < MaxValue);
 
-	DisplayRange.SetLowerBoundValue(MinValue);
-	DisplayRange.SetUpperBoundValue(MaxValue);
+	DisplayRange.SetLowerBoundValue(FMath::Clamp(MinValue, 0.f, MaxValue));
+	DisplayRange.SetUpperBoundValue(FMath::Clamp(MaxValue, MinValue, 1.f));
 
 	if (OnDisplayRangeUpdated.IsBound())
 	{
diff --git a/Engine/Plugins/Runtime/AudioWidgets/Source/AudioWidgets/Public/FixedSampledSequenceGridData.h b/Engine/Plugins/Runtime/AudioWidgets/Source/AudioWidgets/Public/FixedSampledSequenceGridData.h
new file mode 100644
--- /dev/null
+++ b/Engine/Plugins/Runtime/AudioWidgets/Source/AudioWidgets/Public/FixedSampledSequenceGridData.h
@@ -0,0 +1,36 @@
+// Copyright Epic Games, Inc. All Rights Reserved.
+
+#pragma once
+
+#include "Delegates/Delegate.h"
+#include "HAL/Platform.h"
+#include "IFixedSampledSequenceGridService.h"
+#include "Math/Range.h"
+#include "Misc/FrameRate.h"
+
+DECLARE_MULTICAST_DELEGATE(FOnWaveformEditorGridUpdated)
+
+struct FSlateFontInfo;
+
+class FFixedSampledSequenceGridData : public IFixedSampledSequenceGridService
+{
+public:
+	explicit FFixedSampledSequenceGridData(const uint32 InTotalFrames, const uint32 InSampleRateHz, const float InGridSizePixels = 1.f, const FSlateFontInfo* InTicksTimeFont = nullptr);
+
+	void UpdateDisplayRange(const TRange<uint32> InDisplayRange);
+	bool UpdateGridMetrics(const float InGridSizePixels);
+	virtual const FFixedSampledSequenceGridMetrics GetGridMetrics() const override;
+	void SetTicksTimeFont(const FSlateFontInfo* InNewFont);
+	const float SnapPositionToClosestFrame(const float InPixelPosition) const;
+
+	FOnWaveformEditorGridUpdated OnGridMetricsUpdated;
+
+private:
+	FFixedSampledSequenceGridMetrics GridMetrics;
+	uint32 TotalFrames = 0;
+	TRange<uint32> DisplayRange;
+
+	float GridSizePixels = 0.f;
+	const FSlateFontInfo* TicksTimeFont = nullptr;
+	FFrameRate GridFrameRate;
+};
\ No newline at end of file
diff --git a/Engine/Plugins/Runtime/AudioWidgets/Source/AudioWidgets/Public/FixedSampledSequenceView.h b/Engine/Plugins/Runtime/AudioWidgets/Source/AudioWidgets/Public/FixedSampledSequenceView.h
new file mode 100644
--- /dev/null
+++ b/Engine/Plugins/Runtime/AudioWidgets/Source/AudioWidgets/Public/FixedSampledSequenceView.h
@@ -0,0 +1,12 @@
+// Copyright Epic Games, Inc. All Rights Reserved.
+
+#pragma once
+
+#include "Containers/ArrayView.h"
+
+struct FFixedSampledSequenceView
+{
+	TArrayView<const float> SampleData;
+	uint32 NumDimensions = 0;
+	uint32 SampleRate = 0;
+};
\ No newline at end of file
diff --git a/Engine/Plugins/Runtime/AudioWidgets/Source/AudioWidgets/Public/IFixedSampledSequenceViewProvider.h b/Engine/Plugins/Runtime/AudioWidgets/Source/AudioWidgets/Public/IFixedSampledSequenceViewProvider.h
new file mode 100644
--- /dev/null
+++ b/Engine/Plugins/Runtime/AudioWidgets/Source/AudioWidgets/Public/IFixedSampledSequenceViewProvider.h
@@ -0,0 +1,13 @@
+// Copyright Epic Games, Inc. All Rights Reserved.
+
+#pragma once
+
+#include "Math/Range.h"
+#include "FixedSampledSequenceView.h"
+
+class IFixedSampledSequenceViewProvider
+{
+public:
+	virtual ~IFixedSampledSequenceViewProvider() = default;
+	virtual FFixedSampledSequenceView RequestSequenceView(const TRange<double> DataRatioRange) = 0;
+};
\ No newline at end of file
diff --git a/Engine/Plugins/Runtime/AudioWidgets/Source/AudioWidgets/Public/IFixedSampledSequenceViewReceiver.h b/Engine/Plugins/Runtime/AudioWidgets/Source/AudioWidgets/Public/IFixedSampledSequenceViewReceiver.h
new file mode 100644
--- /dev/null
+++ b/Engine/Plugins/Runtime/AudioWidgets/Source/AudioWidgets/Public/IFixedSampledSequenceViewReceiver.h
@@ -0,0 +1,18 @@
+// Copyright Epic Games, Inc. All Rights Reserved.
+
+#pragma once
+
+#include "FixedSampledSequenceView.h"
+
+class IFixedSampledSequenceViewReceiver
+{
+public:
+	virtual ~IFixedSampledSequenceViewReceiver() = default;
+
+	
+	// Receives and handles a sequence view. 
+	// 
+	// @param InView Received view for a Fixed Sampled Sequence
+	// @param FirstSampleIndex the index of the first sample shown by the view	
+	virtual void ReceiveSequenceView(const FFixedSampledSequenceView InView, const uint32 FirstSampleIndex = 0) = 0;
+};
\ No newline at end of file
diff --git a/Engine/Plugins/Runtime/AudioWidgets/Source/AudioWidgets/Public/SparseSampledSequenceTransportCoordinator.h b/Engine/Plugins/Runtime/AudioWidgets/Source/AudioWidgets/Public/SparseSampledSequenceTransportCoordinator.h
--- a/Engine/Plugins/Runtime/AudioWidgets/Source/AudioWidgets/Public/SparseSampledSequenceTransportCoordinator.h
+++ b/Engine/Plugins/Runtime/AudioWidgets/Source/AudioWidgets/Public/SparseSampledSequenceTransportCoordinator.h
@@ -6,16 +6,19 @@
 #include "ISparseSampledSequenceTransportCoordinator.h"
 
 DECLARE_MULTICAST_DELEGATE_OneParam(FOnDisplayRangeUpdated, const TRange<float> /* New Display Range */);
+DECLARE_MULTICAST_DELEGATE_OneParam(FOnFocusPointMoved, const float /* New Position */);
 DECLARE_MULTICAST_DELEGATE_TwoParams(FOnFocusPointScrubUpdate, const float /* Focused Playback Ratio */, const bool /*Playhead is Moving*/);
 
-
 class AUDIOWIDGETS_API FSparseSampledSequenceTransportCoordinator : public ISparseSampledSequenceTransportCoordinator
 {
 public:
 	FSparseSampledSequenceTransportCoordinator() = default;
 	virtual ~FSparseSampledSequenceTransportCoordinator() = default;
 
-	/** Called when the playhead is scrubbed */
+	/** Called when the focus point is set to a new location */
+	FOnFocusPointMoved OnFocusPointMoved;
+
+	/** Called when the focus point is scrubbed */
 	FOnFocusPointScrubUpdate OnFocusPointScrubUpdate;
 
 	/** Called when the display range is updated */
@@ -30,29 +33,26 @@ public:
 	void SetZoomRatio(const float NewRatio) override;
 	float ConvertAbsoluteRatioToZoomed(const float InAbsoluteRatio) const override;
 	float ConvertZoomedRatioToAbsolute(const float InZoomedRatio) const override;
-
 	void UpdatePlaybackRange(const TRange<float>& NewRange);
 	void Stop();
 
 private:
 	FORCEINLINE void MoveFocusPoint(const float InFocusPoint);
 	void UpdateZoomRatioAndDisplayRange(const float NewZoomRatio);
-
 	void UpdateDisplayRange(const float MinValue, const float MaxValue);
 	bool IsRatioWithinDisplayRange(const float Ratio) const;
-
 	float GetPlayBackRatioFromFocusPoint(const float InFocusPoint) const;
-	
+
 	float CurrentPlaybackRatio = 0.f;
 	float FocusPointLockPosition = 0.95f;
 	float FocusPoint = 0.f;
 	float ZoomRatio = 1.f;
-	
+
 	/* The currently displayed render data range */
 	TRange<float> DisplayRange = TRange<float>::Inclusive(0.f, 1.f);
 
 	/* Progress range to scale the incoming progress ratio with*/
 	TRange<float> ProgressRange = TRange<float>::Inclusive(0.f, 1.f);
-	
+
 	bool bIsScrubbing = false;
 };
\ No newline at end of file
