[Waveform Editor] Expose option for asset reimport in toolbar and update view

diff --git a/Engine/Plugins/Editor/WaveformEditor/Source/WaveformEditor/Private/WaveformEditor.cpp b/Engine/Plugins/Editor/WaveformEditor/Source/WaveformEditor/Private/WaveformEditor.cpp
--- a/Engine/Plugins/Editor/WaveformEditor/Source/WaveformEditor/Private/WaveformEditor.cpp
+++ b/Engine/Plugins/Editor/WaveformEditor/Source/WaveformEditor/Private/WaveformEditor.cpp
@@ -3,7 +3,9 @@
 #include "WaveformEditor.h"
 
 #include "AudioDevice.h"
+#include "AssetDefinitionRegistry.h"
 #include "Components/AudioComponent.h"
+#include "EditorReimportHandler.h"
 #include "Misc/TransactionObjectEvent.h"
 #include "PropertyEditorModule.h"
 #include "Sound/SoundWave.h"
@@ -33,6 +35,7 @@ const FName FWaveformEditor::WaveformDisplayTabId("WaveformEditor_Display");
 const FName FWaveformEditor::EditorName("Waveform Editor");
 const FName FWaveformEditor::ToolkitFName("WaveformEditor");
 
+
 bool FWaveformEditor::Init(const EToolkitMode::Type Mode, const TSharedPtr<IToolkitHost>& InitToolkitHost, USoundWave* SoundWaveToEdit)
 {
 	checkf(SoundWaveToEdit, TEXT("Tried to open a Soundwave Editor from a null soundwave"));
@@ -44,11 +47,14 @@ bool FWaveformEditor::Init(const EToolkitMode::Type Mode, const TSharedPtr<ITool
 	bool bIsInitialized = true;
 	
 	bIsInitialized &= CreateDetailsViews();
+	bIsInitialized &= CreateTransportCoordinator();
+	bIsInitialized &= InitializeZoom();
 	bIsInitialized &= CreateWaveformView();
 	bIsInitialized &= InitializeAudioComponent();
 	bIsInitialized &= CreateTransportController();
 	bIsInitialized &= CreateWaveWriter();
 	bIsInitialized &= BindDelegates();
+	bIsInitialized &= SetUpAssetReimport();
 
 	bIsInitialized &= RegisterToolbar();
 	bIsInitialized &= BindCommands();
@@ -79,6 +85,14 @@ bool FWaveformEditor::Init(const EToolkitMode::Type Mode, const TSharedPtr<ITool
 	return bIsInitialized;
 }
 
+FWaveformEditor::~FWaveformEditor()
+{
+	if (FReimportManager::Instance())
+	{
+		FReimportManager::Instance()->OnPostReimport().RemoveAll(this);
+	}
+}
+
 void FWaveformEditor::AddDefaultTransformations()
 {
 	if (!ensure(SoundWave))
@@ -162,6 +176,10 @@ bool FWaveformEditor::CreateTransportController()
 bool FWaveformEditor::InitializeZoom()
 {
 	ZoomManager = MakeShared<FWaveformEditorZoomController>();
+
+	check(TransportCoordinator)
+	ZoomManager->OnZoomRatioChanged.AddSP(TransportCoordinator.ToSharedRef(), &FSparseSampledSequenceTransportCoordinator::SetZoomRatio);
+
 	return ZoomManager != nullptr;
 }
 
@@ -179,6 +197,31 @@ bool FWaveformEditor::BindDelegates()
 	return true;
 }
 
+bool FWaveformEditor::SetUpAssetReimport()
+{
+	if (!FReimportManager::Instance())
+	{
+		return false;
+	}
+
+	FReimportManager::Instance()->OnPostReimport().AddSP(this, &FWaveformEditor::OnAssetReimport);
+	return true;
+}
+
+void FWaveformEditor::ExecuteReimport()
+{
+	if (!CanExecuteReimport())
+	{
+		return;
+	}
+
+	const bool bSelectNewAsset = ReimportMode == EWaveEditorReimportMode::SelectFile;
+
+	TArray<UObject*> CopyOfSelectedAssets;
+	CopyOfSelectedAssets.Add(SoundWave);
+	FReimportManager::Instance()->ValidateAllSourceFileAndReimport(CopyOfSelectedAssets, true, -1, bSelectNewAsset);
+}
+
 void FWaveformEditor::RegisterTabSpawners(const TSharedRef<FTabManager>& InTabManager)
 {
 	WorkspaceMenuCategory = InTabManager->AddLocalWorkspaceMenuCategory(LOCTEXT("WorkspaceMenu_WaveformEditor", "Sound Wave Editor"));
@@ -280,16 +323,26 @@ bool FWaveformEditor::RegisterToolbar()
 		FToolMenuInsert InsertAfterZoomSection("Zoom Controls", EToolMenuInsertType::After);
 		FToolMenuSection& ExportSection = ToolBar->AddSection("Export Controls", TAttribute<FText>(), InsertAfterZoomSection);
 
-		FToolMenuEntry ExportEntry = FToolMenuEntry::InitToolBarButton(
-			Commands.ExportWaveform,
-			LOCTEXT("WaveformEditorRender", ""),
-			LOCTEXT("WaveformEditorRenderButtonTooltip", "Exports the edited waveform to a USoundWave asset"),
-			FSlateIcon(FAppStyle::GetAppStyleSetName(), "LevelEditor.ExportAll")
-		);	
 
-		ExportSection.AddEntry(ExportEntry);
+		ExportSection.AddDynamicEntry("ExportButton", FNewToolMenuSectionDelegate::CreateLambda([this, Commands](FToolMenuSection& InSection) 
+		{
+			UWaveformEditorToolMenuContext* Context = InSection.FindContext<UWaveformEditorToolMenuContext>();
+
+			if (Context && Context->WaveformEditor.IsValid())
+			{
+				FToolMenuEntry ExportEntry = FToolMenuEntry::InitToolBarButton(
+					Commands.ExportWaveform,
+					LOCTEXT("WaveformEditorRender", ""),
+					TAttribute< FText >::CreateRaw(Context->WaveformEditor.Pin().Get(), &FWaveformEditor::GetExportButtonToolTip),
+					FSlateIcon(FAppStyle::GetAppStyleSetName(), "LevelEditor.ExportAll")
+				);
+
+				InSection.AddEntry(ExportEntry);
+			}
+		}));
 
-		ExportSection.AddDynamicEntry("ExportOptions", FNewToolMenuSectionDelegate::CreateLambda([this](FToolMenuSection& InSection) {
+		ExportSection.AddDynamicEntry("ExportOptions", FNewToolMenuSectionDelegate::CreateLambda([this](FToolMenuSection& InSection) 
+		{
 				UWaveformEditorToolMenuContext* Context = InSection.FindContext<UWaveformEditorToolMenuContext>();
 				
 				if (Context && Context->WaveformEditor.IsValid())
@@ -306,9 +359,44 @@ bool FWaveformEditor::RegisterToolbar()
 
 					InSection.AddEntry(ExportOptionsEntry);
 				}				
+		}));
+
+		ExportSection.AddDynamicEntry("ImportButton", FNewToolMenuSectionDelegate::CreateLambda([this, Commands](FToolMenuSection& InSection)
+		{
+			UWaveformEditorToolMenuContext* Context = InSection.FindContext<UWaveformEditorToolMenuContext>();
+
+			if (Context && Context->WaveformEditor.IsValid())
+			{
+				FToolMenuEntry ExportEntry = FToolMenuEntry::InitToolBarButton(
+					Commands.ExportWaveform,
+					LOCTEXT("WaveformEditorReimport", ""),
+					TAttribute< FText >::CreateRaw(Context->WaveformEditor.Pin().Get(), &FWaveformEditor::GetReimportButtonToolTip),
+					FSlateIcon(FAppStyle::GetAppStyleSetName(), "Icons.Import"));
+
+				InSection.AddEntry(ExportEntry);
 			}
-		));
+		}));
+
+		ExportSection.AddDynamicEntry("ReimportOptions", FNewToolMenuSectionDelegate::CreateLambda([this](FToolMenuSection& InSection) 
+		{
+			UWaveformEditorToolMenuContext* Context = InSection.FindContext<UWaveformEditorToolMenuContext>();
 
+			if (Context && Context->WaveformEditor.IsValid())
+			{
+				FToolMenuEntry ReimportOptionsEntry = FToolMenuEntry::InitComboButton(
+					"ReimportOptionsCombo",
+					FToolUIActionChoice(FUIAction()),
+					FNewToolMenuChoice(FOnGetContent::CreateSP(Context->WaveformEditor.Pin().Get(), &FWaveformEditor::GenerateImportOptionsMenu)),
+					LOCTEXT("ReimportOptions_Label", "Reimport Options"),
+					LOCTEXT("ReimportOptions_ToolTip", "Reimport options for this USoundWave"),
+					FSlateIcon(FAppStyle::GetAppStyleSetName(), "Icons.Import"),
+					true
+				);
+
+				InSection.AddEntry(ReimportOptionsEntry);
+			}
+
+		}));
 	}
 
 	return true;
@@ -329,6 +417,31 @@ TSharedRef<SWidget> FWaveformEditor::GenerateExportOptionsMenu()
 	return MenuBuilder.MakeWidget();
 }
 
+TSharedRef<SWidget> FWaveformEditor::GenerateImportOptionsMenu()
+{
+	const bool bShouldCloseWindowAfterMenuSelection = true;
+	FMenuBuilder MenuBuilder(bShouldCloseWindowAfterMenuSelection, GetToolkitCommands());
+
+	MenuBuilder.BeginSection(NAME_None, LOCTEXT("ReimportMode_Label", "Reimport Mode"));
+	{
+		MenuBuilder.AddMenuEntry(FWaveformEditorCommands::Get().ReimportModeSameFile);
+		MenuBuilder.AddMenuEntry(FWaveformEditorCommands::Get().ReimportModeNewFile);
+	}
+	MenuBuilder.EndSection();
+
+	return MenuBuilder.MakeWidget();
+}
+
+bool FWaveformEditor::CanExecuteReimport() const
+{
+	if (FReimportManager::Instance() == nullptr)
+	{
+		return false;
+	}
+
+	return FReimportManager::Instance()->CanReimport(SoundWave);
+}
+
 bool FWaveformEditor::BindCommands()
 {
 	const FWaveformEditorCommands& Commands = FWaveformEditorCommands::Get();
@@ -379,6 +492,23 @@ bool FWaveformEditor::BindCommands()
 		FCanExecuteAction::CreateLambda([this] {return WaveWriter.IsValid(); }),
 		FIsActionChecked::CreateLambda([this] {return WaveWriter->GetExportChannelsFormat() == WaveformEditorWaveWriter::EChannelFormat::Stereo; }));
 
+	ToolkitCommands->MapAction(
+		Commands.ReimportAsset,
+		FExecuteAction::CreateSP(this, &FWaveformEditor::ExecuteReimport),
+		FCanExecuteAction::CreateSP(this, &FWaveformEditor::CanExecuteReimport));
+
+	ToolkitCommands->MapAction(
+		Commands.ReimportModeSameFile,
+		FExecuteAction::CreateLambda([this] { ReimportMode = EWaveEditorReimportMode::SameFile; ExecuteReimport(); }),
+		FCanExecuteAction::CreateLambda([this] { return CanExecuteReimport(); }),
+		FIsActionChecked::CreateLambda([this] {return ReimportMode == EWaveEditorReimportMode::SameFile; }));
+
+	ToolkitCommands->MapAction(
+		Commands.ReimportModeNewFile,
+		FExecuteAction::CreateLambda([this] { ReimportMode = EWaveEditorReimportMode::SelectFile; ExecuteReimport(); }),
+		FCanExecuteAction::CreateLambda([this] { return CanExecuteReimport(); }),
+		FIsActionChecked::CreateLambda([this] {return ReimportMode == EWaveEditorReimportMode::SelectFile; }));
+
 	return true;
 }
 
@@ -408,6 +538,22 @@ FLinearColor FWaveformEditor::GetWorldCentricTabColorScale() const
 	return FLinearColor(0.0f, 0.0f, 0.2f, 0.5f);
 }
 
+void FWaveformEditor::OnAssetReimport(UObject* ReimportedObject, bool bSuccessfullReimport)
+{
+	if (!bSuccessfullReimport)
+	{
+		return;
+	}
+
+	if (ReimportedObject->IsA<USoundWave>() && ReimportedObject->GetPathName() == SoundWave->GetPathName())
+	{
+		CreateWaveformView();
+		WaveformView.DataProvider->RequestSequenceView(TransportCoordinator->GetDisplayRange());
+		WaveformView.ViewWidget->SetPlayheadRatio(TransportCoordinator->GetFocusPoint());
+		this->TabManager->FindExistingLiveTab(WaveformDisplayTabId)->SetContent(WaveformView.ViewWidget.ToSharedRef());
+	}
+}
+
 void FWaveformEditor::NotifyPostChange(const FPropertyChangedEvent& PropertyChangedEvent, class FEditPropertyChain* PropertyThatChanged)
 {	
 	TDoubleLinkedList<FProperty*>::TDoubleLinkedListNode* PropertyNode = PropertyThatChanged->GetActiveMemberNode();
@@ -620,16 +766,15 @@ bool FWaveformEditor::CreateWaveformView()
 		return false;
 	}
 
-	if (!InitializeZoom())
-	{
-		UE_LOG(LogWaveformEditor, Warning, TEXT("Failed to Init Zoom Manager while setting up waveform panel"));
-		return false;
-	}
-
-	TransportCoordinator = MakeShared<FSparseSampledSequenceTransportCoordinator>();
+	check(ZoomManager)
+	check(TransportCoordinator)
 
 	WaveformView = FTransformedWaveformViewFactory::Get().GetTransformedView(SoundWave, TransportCoordinator.ToSharedRef(), this, ZoomManager);
 
+
+	WaveformView.DataProvider->OnRenderElementsUpdated.AddSP(this, &FWaveformEditor::HandleRenderDataUpdate);
+	TransportCoordinator->OnFocusPointMoved.AddSP(WaveformView.ViewWidget.Get(), &STransformedWaveformViewPanel::SetPlayheadRatio);
+
 	WaveformView.DataProvider->OnRenderElementsUpdated.AddSP(this, &FWaveformEditor::HandleRenderDataUpdate);
 	
 	ZoomManager->OnZoomRatioChanged.AddSP(TransportCoordinator.ToSharedRef(), &FSparseSampledSequenceTransportCoordinator::SetZoomRatio);
@@ -640,6 +785,14 @@ bool FWaveformEditor::CreateWaveformView()
 	return WaveformView.ViewWidget != nullptr && WaveformView.DataProvider != nullptr;
 }
 
+bool FWaveformEditor::CreateTransportCoordinator()
+{
+	TransportCoordinator = MakeShared<FSparseSampledSequenceTransportCoordinator>();
+	TransportCoordinator->OnDisplayRangeUpdated.AddSP(this, &FWaveformEditor::HandleDisplayRangeUpdate);
+
+	return TransportCoordinator != nullptr;
+}
+
 void FWaveformEditor::HandlePlaybackPercentageChange(const UAudioComponent* InComponent, const USoundWave* InSoundWave, const float InPlaybackPercentage)
 {
 	const bool bIsStopped = AudioComponent->GetPlayState() == EAudioComponentPlayState::Stopped;
@@ -769,4 +922,49 @@ const UWaveformEditorTransformationsSettings* FWaveformEditor::GetWaveformEditor
 	return WaveformEditorTransformationsSettings;
 }
 
+FText FWaveformEditor::GetReimportButtonToolTip() const
+{
+	FText ReimportModeText;
+
+	switch (ReimportMode)
+	{
+	case(EWaveEditorReimportMode::SelectFile):
+		ReimportModeText = LOCTEXT("SelectFile", "Reimport from new file");
+		break;
+	case(EWaveEditorReimportMode::SameFile):
+		ReimportModeText = LOCTEXT("SameFile", "Reimport from same file");
+		break;
+	default:
+		static_assert(static_cast<int32>(EWaveEditorReimportMode::COUNT) == 2, "Possible missing switch case coverage for 'EWaveEditorReimportMode'");
+		break;
+	}
+
+	return FText::Format(LOCTEXT("WaveformEditorReimportButtonTooltip", "{0}."), ReimportModeText);
+}
+
+FText FWaveformEditor::GetExportButtonToolTip() const
+{
+
+	FText ExportModeText;
+
+	if (!WaveWriter)
+	{
+		return ExportModeText;
+	}
+
+	switch (WaveWriter->GetExportChannelsFormat())
+	{
+	case(WaveformEditorWaveWriter::EChannelFormat::Stereo):
+		ExportModeText = LOCTEXT("ExportModeStereo", "stereo");
+		break;
+	case(WaveformEditorWaveWriter::EChannelFormat::Mono):
+		ExportModeText = LOCTEXT("SameFile", "mono");
+		break;
+	default:
+		break;
+	}
+
+	return FText::Format(LOCTEXT("WaveformEditorExportButtonTooltip", "Exports the edited waveform to a {0} USoundWave asset."), ExportModeText);
+}
+
 #undef LOCTEXT_NAMESPACE
\ No newline at end of file
diff --git a/Engine/Plugins/Editor/WaveformEditor/Source/WaveformEditor/Private/WaveformEditorCommands.cpp b/Engine/Plugins/Editor/WaveformEditor/Source/WaveformEditor/Private/WaveformEditorCommands.cpp
--- a/Engine/Plugins/Editor/WaveformEditor/Source/WaveformEditor/Private/WaveformEditorCommands.cpp
+++ b/Engine/Plugins/Editor/WaveformEditor/Source/WaveformEditor/Private/WaveformEditorCommands.cpp
@@ -26,6 +26,10 @@ void FWaveformEditorCommands::RegisterCommands()
 	UI_COMMAND(ExportWaveform, "Export Waveform", "Exports the edited waveform to a new sound wave uasset", EUserInterfaceActionType::Button, FInputChord());
 	UI_COMMAND(ExportFormatMono, "Mono", "Sets the export format to mono", EUserInterfaceActionType::RadioButton, FInputChord());
 	UI_COMMAND(ExportFormatStereo, "Stereo", "Sets the export format to stereo", EUserInterfaceActionType::RadioButton, FInputChord());
+
+	UI_COMMAND(ReimportAsset, "Reimport Asset", "Reimports the inspected sound wave uasset", EUserInterfaceActionType::Button, FInputChord());
+	UI_COMMAND(ReimportModeSameFile, "Reimport from same file", "Reimports asset from the same file", EUserInterfaceActionType::RadioButton, FInputChord());
+	UI_COMMAND(ReimportModeNewFile, "Reimport from new file", "Reimports asset from a selected file", EUserInterfaceActionType::RadioButton, FInputChord());
 }
 	
 const FWaveformEditorCommands& FWaveformEditorCommands::Get()
diff --git a/Engine/Plugins/Editor/WaveformEditor/Source/WaveformEditor/Public/WaveformEditor.h b/Engine/Plugins/Editor/WaveformEditor/Source/WaveformEditor/Public/WaveformEditor.h
--- a/Engine/Plugins/Editor/WaveformEditor/Source/WaveformEditor/Public/WaveformEditor.h
+++ b/Engine/Plugins/Editor/WaveformEditor/Source/WaveformEditor/Public/WaveformEditor.h
@@ -25,6 +25,7 @@ class WAVEFORMEDITOR_API FWaveformEditor
 public:
 
 	bool Init(const EToolkitMode::Type Mode, const TSharedPtr<IToolkitHost>& InitToolkitHost, USoundWave* SoundWaveToEdit);
+	virtual ~FWaveformEditor();
 
 	/** FAssetEditorToolkit interface */
 	virtual void RegisterTabSpawners(const TSharedRef<FTabManager>& TabManager) override;
@@ -35,6 +36,8 @@ public:
 	virtual FString GetWorldCentricTabPrefix() const override;
 	virtual FLinearColor GetWorldCentricTabColorScale() const override;
 
+	void OnAssetReimport(UObject* ReimportedObject, bool bSuccessfullReimport);
+
 	/** FNotifyHook interface */
 	void NotifyPreChange(class FEditPropertyChain* PropertyAboutToChange) override {};
 	void NotifyPostChange(const FPropertyChangedEvent& PropertyChangedEvent, class FEditPropertyChain* PropertyThatChanged) override;
@@ -50,6 +53,8 @@ private:
 	bool CreateTransportController();
 	bool InitializeZoom();
 	bool BindDelegates();
+	bool SetUpAssetReimport();
+	void ExecuteReimport();
 
 	/**	Sets the wave editor layout */
 	const TSharedRef<FTabManager::FLayout> SetupStandaloneLayout();
@@ -58,6 +63,8 @@ private:
 	bool RegisterToolbar();
 	bool BindCommands();
 	TSharedRef<SWidget> GenerateExportOptionsMenu();
+	TSharedRef<SWidget> GenerateImportOptionsMenu();
+	bool CanExecuteReimport() const;
 
 	/**	Details tabs set up */
 	TSharedRef<SDockTab> SpawnTab_Properties(const FSpawnTabArgs& Args);
@@ -67,13 +74,14 @@ private:
 	/**	Waveform view tab setup */
 	TSharedRef<SDockTab> SpawnTab_WaveformDisplay(const FSpawnTabArgs& Args);
 	bool CreateWaveformView();
+	bool CreateTransportCoordinator();
 
 	/** Playback delegates handlers */
 	void HandlePlaybackPercentageChange(const UAudioComponent* InComponent, const USoundWave* InSoundWave, const float InPlaybackPercentage);
 	void HandleAudioComponentPlayStateChanged(const UAudioComponent* InAudioComponent, EAudioComponentPlayState NewPlayState);
 	void HandlePlayheadScrub(const float InTargetPlayBackRatio, const bool bIsMoving);
 
-	/* Delegates Handlers*/
+	/* Data View Delegates Handlers*/
 	void HandleRenderDataUpdate();
 	void HandleDisplayRangeUpdate(const TRange<double>);
 
@@ -124,5 +132,16 @@ private:
 	float PlaybackTimeBeforeTransformInteraction = 0.f;
 	float StartTimeBeforeTransformInteraction = 0.f;
 	FWaveTransformUObjectConfiguration TransformationChainConfig;
+
+	enum class EWaveEditorReimportMode : uint8
+	{
+		SameFile = 0, 
+		SelectFile,
+		COUNT
+
+	} ReimportMode;
+
+	FText GetReimportButtonToolTip() const;
+	FText GetExportButtonToolTip() const;
 };
 
diff --git a/Engine/Plugins/Editor/WaveformEditor/Source/WaveformEditor/Public/WaveformEditorCommands.h b/Engine/Plugins/Editor/WaveformEditor/Source/WaveformEditor/Public/WaveformEditorCommands.h
--- a/Engine/Plugins/Editor/WaveformEditor/Source/WaveformEditor/Public/WaveformEditorCommands.h
+++ b/Engine/Plugins/Editor/WaveformEditor/Source/WaveformEditor/Public/WaveformEditorCommands.h
@@ -26,4 +26,9 @@ public:
 	TSharedPtr<FUICommandInfo> ExportWaveform;
 	TSharedPtr<FUICommandInfo> ExportFormatMono;
 	TSharedPtr<FUICommandInfo> ExportFormatStereo;
+
+	TSharedPtr<FUICommandInfo> ReimportAsset;
+	TSharedPtr<FUICommandInfo> ReimportModeSameFile;
+	TSharedPtr<FUICommandInfo> ReimportModeNewFile;
+
 };
