Localization fixes

diff --git a/Engine/Plugins/Editor/WaveformEditor/Source/WaveformEditor/Private/WaveformEditor.cpp b/Engine/Plugins/Editor/WaveformEditor/Source/WaveformEditor/Private/WaveformEditor.cpp
--- a/Engine/Plugins/Editor/WaveformEditor/Source/WaveformEditor/Private/WaveformEditor.cpp
+++ b/Engine/Plugins/Editor/WaveformEditor/Source/WaveformEditor/Private/WaveformEditor.cpp
@@ -958,7 +958,7 @@ FText FWaveformEditor::GetExportButtonToolTip() const
 		ExportModeText = LOCTEXT("ExportModeStereo", "stereo");
 		break;
 	case(WaveformEditorWaveWriter::EChannelFormat::Mono):
-		ExportModeText = LOCTEXT("SameFile", "mono");
+		ExportModeText = LOCTEXT("ExportModeMono", "mono");
 		break;
 	default:
 		break;
@@ -967,4 +967,4 @@ FText FWaveformEditor::GetExportButtonToolTip() const
 	return FText::Format(LOCTEXT("WaveformEditorExportButtonTooltip", "Exports the edited waveform to a {0} USoundWave asset."), ExportModeText);
 }
 
-#undef LOCTEXT_NAMESPACE
\ No newline at end of file
+#undef LOCTEXT_NAMESPACE
diff --git a/Engine/Plugins/VirtualProduction/StageMonitoring/Source/StageMonitorEditor/Private/Widgets/SDataProviderListView.cpp b/Engine/Plugins/VirtualProduction/StageMonitoring/Source/StageMonitorEditor/Private/Widgets/SDataProviderListView.cpp
--- a/Engine/Plugins/VirtualProduction/StageMonitoring/Source/StageMonitorEditor/Private/Widgets/SDataProviderListView.cpp
+++ b/Engine/Plugins/VirtualProduction/StageMonitoring/Source/StageMonitorEditor/Private/Widgets/SDataProviderListView.cpp
@@ -336,7 +336,7 @@ namespace DataProviderListView
 			},
 			{GameThreadWaitTiming,
 			 {
-				 LOCTEXT("HeaderName_GameThread", "Game Wait Time (ms)"), IsNotSortable, IsNotVisible,
+				 LOCTEXT("HeaderName_GameWaitThread", "Game Wait Time (ms)"), IsNotSortable, IsNotVisible,
 				 [](SDataProviderTableRow* Provider) -> TSharedRef<SWidget>
 				 {
 					 return SNew(SHorizontalBox)
