fix wave editor localization erorrs in build CL 25100598
#trivial

diff --git a/Engine/Plugins/Editor/WaveformEditor/Source/WaveformEditorWidgets/Private/STransformedWaveformViewPanel.cpp b/Engine/Plugins/Editor/WaveformEditor/Source/WaveformEditorWidgets/Private/STransformedWaveformViewPanel.cpp
--- a/Engine/Plugins/Editor/WaveformEditor/Source/WaveformEditorWidgets/Private/STransformedWaveformViewPanel.cpp
+++ b/Engine/Plugins/Editor/WaveformEditor/Source/WaveformEditorWidgets/Private/STransformedWaveformViewPanel.cpp
@@ -344,7 +344,7 @@ void STransformedWaveformViewPanel::SetUpValueGridOverlay()
 		const float MinDecibelNumericLimit = -160;
 		FNumberFormattingOptions Formatting; 
 		Formatting.SetMaximumFractionalDigits(1);
-		return FText::Format(LOCTEXT("WaveformEditorValueGridLabel", "{0}"), (GridLineValue <= MinDecibelNumericLimit ? LOCTEXT("WaveformEditorValueGridLabel", "-inf") : FText::AsNumber(GridLineValue, &Formatting)));
+		return FText::Format(LOCTEXT("WaveformEditorValueGridLabel", "{0}"), (GridLineValue <= MinDecibelNumericLimit ? LOCTEXT("WaveformEditorValueGridNegativeInf", "-inf") : FText::AsNumber(GridLineValue, &Formatting)));
 	};
 
 	const UWaveformEditorWidgetsSettings* Settings = GetWaveformEditorWidgetsSettings();
