[Waveform Editor] Restore time ruler display unit selection menu

diff --git a/Engine/Plugins/Editor/WaveformEditor/Source/WaveformEditor/Private/TransformedWaveformViewFactory.cpp b/Engine/Plugins/Editor/WaveformEditor/Source/WaveformEditor/Private/TransformedWaveformViewFactory.cpp
--- a/Engine/Plugins/Editor/WaveformEditor/Source/WaveformEditor/Private/TransformedWaveformViewFactory.cpp
+++ b/Engine/Plugins/Editor/WaveformEditor/Source/WaveformEditor/Private/TransformedWaveformViewFactory.cpp
@@ -119,6 +119,13 @@ void FTransformedWaveformViewFactory::SetUpWaveformPanelInteractions(TSharedRef<
 	{
 		TSharedRef<SWidget> WavePanelChildSlot = WaveformPanel->GetChildren()->GetChildAt(0);
 
+		const bool HandleRightMouseButton = MouseEvent.GetEffectingButton() == EKeys::RightMouseButton;
+		
+		if (HandleRightMouseButton)
+		{
+			return WaveformPanel->LaunchTimeRulerContextMenu();
+		}
+
 		for (int32 ChildIndex = 0; ChildIndex < WavePanelChildSlot->GetChildren()->Num(); ++ChildIndex)
 		{
 			TSharedRef<SWidget> RulerWidget = WavePanelChildSlot->GetChildren()->GetChildAt(ChildIndex);
diff --git a/Engine/Plugins/Editor/WaveformEditor/Source/WaveformEditorWidgets/Private/STransformedWaveformViewPanel.cpp b/Engine/Plugins/Editor/WaveformEditor/Source/WaveformEditorWidgets/Private/STransformedWaveformViewPanel.cpp
--- a/Engine/Plugins/Editor/WaveformEditor/Source/WaveformEditorWidgets/Private/STransformedWaveformViewPanel.cpp
+++ b/Engine/Plugins/Editor/WaveformEditor/Source/WaveformEditorWidgets/Private/STransformedWaveformViewPanel.cpp
@@ -221,6 +221,16 @@ void STransformedWaveformViewPanel::SetOnMouseWheel(FPointerEventHandler InEvent
 	InputRoutingOverlay->OnMouseWheelDelegate = InEventHandler;
 }
 
+FReply STransformedWaveformViewPanel::LaunchTimeRulerContextMenu()
+{
+	if (TimeRuler)
+	{
+		return TimeRuler->LaunchContextMenu();
+	}
+
+	return FReply::Unhandled();
+}
+
 void STransformedWaveformViewPanel::UpdateDisplayUnit(const ESampledSequenceDisplayUnit InDisplayUnit)
 {
 	DisplayUnit = InDisplayUnit;
diff --git a/Engine/Plugins/Editor/WaveformEditor/Source/WaveformEditorWidgets/Public/STransformedWaveformViewPanel.h b/Engine/Plugins/Editor/WaveformEditor/Source/WaveformEditorWidgets/Public/STransformedWaveformViewPanel.h
--- a/Engine/Plugins/Editor/WaveformEditor/Source/WaveformEditorWidgets/Public/STransformedWaveformViewPanel.h
+++ b/Engine/Plugins/Editor/WaveformEditor/Source/WaveformEditorWidgets/Public/STransformedWaveformViewPanel.h
@@ -55,6 +55,8 @@ public:
 	void SetOnTimeRulerMouseMove(FPointerEventHandler InEventHandler);
 	void SetOnMouseWheel(FPointerEventHandler InEventHandler);
 
+	FReply LaunchTimeRulerContextMenu();
+
 private:
 	void CreateLayout();
 
diff --git a/Engine/Plugins/Runtime/AudioWidgets/Source/AudioWidgets/Private/SFixedSampledSequenceRuler.cpp b/Engine/Plugins/Runtime/AudioWidgets/Source/AudioWidgets/Private/SFixedSampledSequenceRuler.cpp
--- a/Engine/Plugins/Runtime/AudioWidgets/Source/AudioWidgets/Private/SFixedSampledSequenceRuler.cpp
+++ b/Engine/Plugins/Runtime/AudioWidgets/Source/AudioWidgets/Private/SFixedSampledSequenceRuler.cpp
@@ -237,6 +237,12 @@ void SFixedSampledSequenceRuler::OnStyleUpdated(const FNotifyingAudioWidgetStyle
 	TicksTextOffset = Style->TicksTextOffset;
 }
 
+FReply SFixedSampledSequenceRuler::LaunchContextMenu()
+{
+	FSlateApplication::Get().PushMenu(AsShared(), FWidgetPath(), MakeContextMenu(), FSlateApplication::Get().GetCursorPos(), FPopupTransitionEffect::ContextMenu);
+	return FReply::Handled();
+}
+
 void SFixedSampledSequenceRuler::DrawPlayheadHandle(const FGeometry& AllottedGeometry, FSlateWindowElementList& OutDrawElements, int32& LayerId) const
 {
 	const float WindowWidth = AllottedGeometry.Size.X;
diff --git a/Engine/Plugins/Runtime/AudioWidgets/Source/AudioWidgets/Public/SFixedSampledSequenceRuler.h b/Engine/Plugins/Runtime/AudioWidgets/Source/AudioWidgets/Public/SFixedSampledSequenceRuler.h
--- a/Engine/Plugins/Runtime/AudioWidgets/Source/AudioWidgets/Public/SFixedSampledSequenceRuler.h
+++ b/Engine/Plugins/Runtime/AudioWidgets/Source/AudioWidgets/Public/SFixedSampledSequenceRuler.h
@@ -34,6 +34,7 @@ public:
 	void UpdateDisplayUnit(const ESampledSequenceDisplayUnit InDisplayUnit);
 	void SetPlayheadPosition(const float InNewPosition);
 	void OnStyleUpdated(const FNotifyingAudioWidgetStyle& UpdatedStyle);
+	FReply LaunchContextMenu();
 
 	/** Delegate sent when the user selects a new display unit from the RMB menu*/
 	FOnTimeUnitMenuSelection OnTimeUnitMenuSelection;
