[Waveform Editor] Make export button options as check buttons to be consistent with UE UI
#trivial

diff --git a/Engine/Plugins/Editor/WaveformEditor/Source/WaveformEditor/Private/WaveformEditorCommands.cpp b/Engine/Plugins/Editor/WaveformEditor/Source/WaveformEditor/Private/WaveformEditorCommands.cpp
--- a/Engine/Plugins/Editor/WaveformEditor/Source/WaveformEditor/Private/WaveformEditorCommands.cpp
+++ b/Engine/Plugins/Editor/WaveformEditor/Source/WaveformEditor/Private/WaveformEditorCommands.cpp
@@ -28,8 +28,8 @@ void FWaveformEditorCommands::RegisterCommands()
 	UI_COMMAND(ExportFormatStereo, "Stereo", "Sets the export format to stereo", EUserInterfaceActionType::RadioButton, FInputChord());
 
 	UI_COMMAND(ReimportAsset, "Reimport Asset", "Reimports the inspected sound wave uasset", EUserInterfaceActionType::Button, FInputChord());
-	UI_COMMAND(ReimportModeSameFile, "Reimport from same file", "Reimports asset from the same file", EUserInterfaceActionType::RadioButton, FInputChord());
-	UI_COMMAND(ReimportModeNewFile, "Reimport from new file", "Reimports asset from a selected file", EUserInterfaceActionType::RadioButton, FInputChord());
+	UI_COMMAND(ReimportModeSameFile, "Reimport from same file", "Reimports asset from the same file", EUserInterfaceActionType::Check, FInputChord());
+	UI_COMMAND(ReimportModeNewFile, "Reimport from new file", "Reimports asset from a selected file", EUserInterfaceActionType::Check, FInputChord());
 }
 	
 const FWaveformEditorCommands& FWaveformEditorCommands::Get()
