[Waveform Editor] Fix reimport button executing an export
#trivial

diff --git a/Engine/Plugins/Editor/WaveformEditor/Source/WaveformEditor/Private/WaveformEditor.cpp b/Engine/Plugins/Editor/WaveformEditor/Source/WaveformEditor/Private/WaveformEditor.cpp
--- a/Engine/Plugins/Editor/WaveformEditor/Source/WaveformEditor/Private/WaveformEditor.cpp
+++ b/Engine/Plugins/Editor/WaveformEditor/Source/WaveformEditor/Private/WaveformEditor.cpp
@@ -367,13 +367,13 @@ bool FWaveformEditor::RegisterToolbar()
 
 			if (Context && Context->WaveformEditor.IsValid())
 			{
-				FToolMenuEntry ExportEntry = FToolMenuEntry::InitToolBarButton(
-					Commands.ExportWaveform,
+				FToolMenuEntry ReimportEntry = FToolMenuEntry::InitToolBarButton(
+					Commands.ReimportAsset,
 					LOCTEXT("WaveformEditorReimport", ""),
 					TAttribute< FText >::CreateRaw(Context->WaveformEditor.Pin().Get(), &FWaveformEditor::GetReimportButtonToolTip),
 					FSlateIcon(FAppStyle::GetAppStyleSetName(), "Icons.Import"));
 
-				InSection.AddEntry(ExportEntry);
+				InSection.AddEntry(ReimportEntry);
 			}
 		}));
 
