[Waveform Editor] Unbind delegates tied to waveform view when replacing it + remove dependencies to transport coordinator from waveform transformations overlay

fix for #jira UE-184284

diff --git a/Engine/Plugins/Editor/WaveformEditor/Source/WaveformEditor/Private/TransformedWaveformViewFactory.cpp b/Engine/Plugins/Editor/WaveformEditor/Source/WaveformEditor/Private/TransformedWaveformViewFactory.cpp
--- a/Engine/Plugins/Editor/WaveformEditor/Source/WaveformEditor/Private/TransformedWaveformViewFactory.cpp
+++ b/Engine/Plugins/Editor/WaveformEditor/Source/WaveformEditor/Private/TransformedWaveformViewFactory.cpp
@@ -38,14 +38,18 @@ FTransformedWaveformView FTransformedWaveformViewFactory::GetTransformedView(TOb
 
 	FFixedSampledSequenceView SequenceView = WaveformDataProvider->RequestSequenceView(TRange<double>::Inclusive(0, 1));
 
-	TSharedPtr<SWaveformTransformationsOverlay> TransformationsOverlay = SNew(SWaveformTransformationsOverlay, WaveformDataProvider->GetTransformLayers(), TransportCoordinator);
+	TFunction<float(const float)> TransformationsOverlayZoomConverter = [TransportCoordinator](const float InAnchorRatio)
+	{
+		return TransportCoordinator->ConvertAbsoluteRatioToZoomed(InAnchorRatio);
+	};
+
+	TSharedPtr<SWaveformTransformationsOverlay> TransformationsOverlay = SNew(SWaveformTransformationsOverlay, WaveformDataProvider->GetTransformLayers()).AnchorsRatioConverter(TransformationsOverlayZoomConverter);
 
 	TSharedPtr<STransformedWaveformViewPanel> WaveformPanel = SNew(STransformedWaveformViewPanel, SequenceView).TransformationsOverlay(TransformationsOverlay);
 	SetUpWaveformPanelInteractions(WaveformPanel.ToSharedRef(), TransportCoordinator, InZoomController, TransformationsOverlay);
 
 	WaveformDataProvider->OnLayersChainGenerated.AddSP(TransformationsOverlay.Get(), &SWaveformTransformationsOverlay::OnLayerChainGenerated);
 	WaveformDataProvider->OnRenderElementsUpdated.AddSP(TransformationsOverlay.Get(), &SWaveformTransformationsOverlay::UpdateLayerConstraints);
-	TransportCoordinator->OnDisplayRangeUpdated.AddSP(TransformationsOverlay.Get(), &SWaveformTransformationsOverlay::OnNewWaveformDisplayRange);
 	WaveformDataProvider->OnDataViewGenerated.AddSP(WaveformPanel.Get(), &STransformedWaveformViewPanel::ReceiveSequenceView);
 
 	return FTransformedWaveformView{ WaveformPanel, WaveformDataProvider};
diff --git a/Engine/Plugins/Editor/WaveformEditor/Source/WaveformEditor/Private/WaveformEditor.cpp b/Engine/Plugins/Editor/WaveformEditor/Source/WaveformEditor/Private/WaveformEditor.cpp
--- a/Engine/Plugins/Editor/WaveformEditor/Source/WaveformEditor/Private/WaveformEditor.cpp
+++ b/Engine/Plugins/Editor/WaveformEditor/Source/WaveformEditor/Private/WaveformEditor.cpp
@@ -765,24 +765,38 @@ bool FWaveformEditor::CreateWaveformView()
 		UE_LOG(LogWaveformEditor, Warning, TEXT("Trying to setup waveform panel from a null SoundWave"));
 		return false;
 	}
-
-	check(ZoomManager)
-	check(TransportCoordinator)
+	
+	if (WaveformView.IsValid())
+	{
+		RemoveWaveformViewDelegates(*WaveformView.DataProvider, *WaveformView.ViewWidget);
+	}
 
 	WaveformView = FTransformedWaveformViewFactory::Get().GetTransformedView(SoundWave, TransportCoordinator.ToSharedRef(), this, ZoomManager);
 
+	check(ZoomManager)
 
-	WaveformView.DataProvider->OnRenderElementsUpdated.AddSP(this, &FWaveformEditor::HandleRenderDataUpdate);
-	TransportCoordinator->OnFocusPointMoved.AddSP(WaveformView.ViewWidget.Get(), &STransformedWaveformViewPanel::SetPlayheadRatio);
+	BindWaveformViewDelegates(*WaveformView.DataProvider, *WaveformView.ViewWidget);
 
-	WaveformView.DataProvider->OnRenderElementsUpdated.AddSP(this, &FWaveformEditor::HandleRenderDataUpdate);
-	
 	ZoomManager->OnZoomRatioChanged.AddSP(TransportCoordinator.ToSharedRef(), &FSparseSampledSequenceTransportCoordinator::SetZoomRatio);
-
 	TransportCoordinator->OnDisplayRangeUpdated.AddSP(this, &FWaveformEditor::HandleDisplayRangeUpdate);
-	TransportCoordinator->OnFocusPointMoved.AddSP(WaveformView.ViewWidget.Get(), &STransformedWaveformViewPanel::SetPlayheadRatio);
 
-	return WaveformView.ViewWidget != nullptr && WaveformView.DataProvider != nullptr;
+	return WaveformView.IsValid();
+}
+
+void FWaveformEditor::BindWaveformViewDelegates(FWaveformEditorSequenceDataProvider& ViewDataProvider, STransformedWaveformViewPanel& ViewWidget)
+{
+	check(TransportCoordinator)
+
+	ViewDataProvider.OnRenderElementsUpdated.AddSP(this, &FWaveformEditor::HandleRenderDataUpdate);
+	TransportCoordinator->OnFocusPointMoved.AddSP(&ViewWidget, &STransformedWaveformViewPanel::SetPlayheadRatio);
+}
+
+void FWaveformEditor::RemoveWaveformViewDelegates(FWaveformEditorSequenceDataProvider& ViewDataProvider, STransformedWaveformViewPanel& ViewWidget)
+{
+	check(TransportCoordinator)
+
+	ViewDataProvider.OnRenderElementsUpdated.RemoveAll(this);
+	TransportCoordinator->OnFocusPointMoved.RemoveAll(&ViewWidget);
 }
 
 bool FWaveformEditor::CreateTransportCoordinator()
diff --git a/Engine/Plugins/Editor/WaveformEditor/Source/WaveformEditor/Public/TransformedWaveformView.h b/Engine/Plugins/Editor/WaveformEditor/Source/WaveformEditor/Public/TransformedWaveformView.h
--- a/Engine/Plugins/Editor/WaveformEditor/Source/WaveformEditor/Public/TransformedWaveformView.h
+++ b/Engine/Plugins/Editor/WaveformEditor/Source/WaveformEditor/Public/TransformedWaveformView.h
@@ -9,4 +9,9 @@ struct FTransformedWaveformView
 {
 	TSharedPtr<STransformedWaveformViewPanel> ViewWidget;
 	TSharedPtr<FWaveformEditorSequenceDataProvider> DataProvider;
+
+	bool IsValid()
+	{
+		return ViewWidget.IsValid() && DataProvider.IsValid();
+	}
 };
\ No newline at end of file
diff --git a/Engine/Plugins/Editor/WaveformEditor/Source/WaveformEditor/Public/WaveformEditor.h b/Engine/Plugins/Editor/WaveformEditor/Source/WaveformEditor/Public/WaveformEditor.h
--- a/Engine/Plugins/Editor/WaveformEditor/Source/WaveformEditor/Public/WaveformEditor.h
+++ b/Engine/Plugins/Editor/WaveformEditor/Source/WaveformEditor/Public/WaveformEditor.h
@@ -75,6 +75,8 @@ private:
 	TSharedRef<SDockTab> SpawnTab_WaveformDisplay(const FSpawnTabArgs& Args);
 	bool CreateWaveformView();
 	bool CreateTransportCoordinator();
+	void BindWaveformViewDelegates(FWaveformEditorSequenceDataProvider& ViewDataProvider, STransformedWaveformViewPanel& ViewWidget);
+	void RemoveWaveformViewDelegates(FWaveformEditorSequenceDataProvider& ViewDataProvider, STransformedWaveformViewPanel& ViewWidget);
 
 	/** Playback delegates handlers */
 	void HandlePlaybackPercentageChange(const UAudioComponent* InComponent, const USoundWave* InSoundWave, const float InPlaybackPercentage);
diff --git a/Engine/Plugins/Editor/WaveformEditor/Source/WaveformEditorWidgets/Private/STransformedWaveformViewPanel.cpp b/Engine/Plugins/Editor/WaveformEditor/Source/WaveformEditorWidgets/Private/STransformedWaveformViewPanel.cpp
--- a/Engine/Plugins/Editor/WaveformEditor/Source/WaveformEditorWidgets/Private/STransformedWaveformViewPanel.cpp
+++ b/Engine/Plugins/Editor/WaveformEditor/Source/WaveformEditorWidgets/Private/STransformedWaveformViewPanel.cpp
@@ -175,6 +175,11 @@ void STransformedWaveformViewPanel::ReceiveSequenceView(const FFixedSampledSeque
 		GridData->UpdateDisplayRange(TRange<uint32>(FirstRenderedFrame, FirstRenderedFrame + NumFrames));
 	}
 
+	if (WaveformTransformationsOverlay)
+	{
+		WaveformTransformationsOverlay->UpdateLayerConstraints();
+	}
+
 	if (WaveformViewer)
 	{
 		WaveformViewer->UpdateView(InView.SampleData, InView.NumDimensions);
diff --git a/Engine/Plugins/Editor/WaveformEditor/Source/WaveformEditorWidgets/Private/SWaveformTransformationsOverlay.cpp b/Engine/Plugins/Editor/WaveformEditor/Source/WaveformEditorWidgets/Private/SWaveformTransformationsOverlay.cpp
--- a/Engine/Plugins/Editor/WaveformEditor/Source/WaveformEditorWidgets/Private/SWaveformTransformationsOverlay.cpp
+++ b/Engine/Plugins/Editor/WaveformEditor/Source/WaveformEditorWidgets/Private/SWaveformTransformationsOverlay.cpp
@@ -7,10 +7,10 @@
 #include "SparseSampledSequenceTransportCoordinator.h"
 #include "Widgets/SOverlay.h"
 
-void SWaveformTransformationsOverlay::Construct(const FArguments& InArgs, TArrayView<const FTransformationLayerRenderInfo> InTransformationRenderers, TSharedRef<FSparseSampledSequenceTransportCoordinator> InTransportCoordinator)
+void SWaveformTransformationsOverlay::Construct(const FArguments& InArgs, TArrayView<const FTransformationLayerRenderInfo> InTransformationRenderers)
 {
 	TransformationRenderers = InTransformationRenderers;
-	TransportCoordinator = InTransportCoordinator;
+	AnchorsRatioConverter = InArgs._AnchorsRatioConverter;
 	CreateLayout();
 }
 
@@ -33,8 +33,8 @@ void SWaveformTransformationsOverlay::CreateLayout()
 		
 		if (LayerRenderInfo.Key)
 		{
-			float LeftAnchor = TransportCoordinator->ConvertAbsoluteRatioToZoomed(LayerRenderInfo.Value.Key);
-			float RightAnchor = TransportCoordinator->ConvertAbsoluteRatioToZoomed(LayerRenderInfo.Value.Value);
+			float LeftAnchor = AnchorsRatioConverter(LayerRenderInfo.Value.Key);
+			float RightAnchor = AnchorsRatioConverter(LayerRenderInfo.Value.Value);
 
 			SConstraintCanvas::FSlot* SlotPtr = nullptr;
 
@@ -72,8 +72,8 @@ void SWaveformTransformationsOverlay::UpdateAnchors()
 
 		if (Layer.Key)
 		{
-			float LeftAnchor = TransportCoordinator->ConvertAbsoluteRatioToZoomed(Layer.Value.Key);
-			float RightAnchor = TransportCoordinator->ConvertAbsoluteRatioToZoomed(Layer.Value.Value);
+			float LeftAnchor = AnchorsRatioConverter(Layer.Value.Key);
+			float RightAnchor = AnchorsRatioConverter(Layer.Value.Value);
 
 			LayersSlots[i]->SetAnchors(FAnchors(LeftAnchor, 0.f, RightAnchor, 1.f));
 		}
@@ -91,11 +91,6 @@ void SWaveformTransformationsOverlay::UpdateLayerConstraints()
 	UpdateAnchors();
 }
 
-void SWaveformTransformationsOverlay::OnNewWaveformDisplayRange(const TRange<double> NewDisplayRange)
-{
-	UpdateAnchors();
-}
-
 FReply SWaveformTransformationsOverlay::OnMouseButtonDown(const FGeometry& MyGeometry, const FPointerEvent& MouseEvent)
 {
 	return AudioWidgetsUtils::RouteMouseInput(&SWidget::OnMouseButtonDown, MouseEvent, MakeArrayView(TransformationLayers.GetData(), TransformationLayers.Num()), true);
@@ -119,4 +114,4 @@ FReply SWaveformTransformationsOverlay::OnMouseWheel(const FGeometry& MyGeometry
 FCursorReply SWaveformTransformationsOverlay::OnCursorQuery(const FGeometry& MyGeometry, const FPointerEvent& CursorEvent) const
 {
 	return AudioWidgetsUtils::RouteCursorQuery(CursorEvent, MakeArrayView(TransformationLayers.GetData(), TransformationLayers.Num()), true);
-}
\ No newline at end of file
+}
diff --git a/Engine/Plugins/Editor/WaveformEditor/Source/WaveformEditorWidgets/Public/SWaveformTransformationsOverlay.h b/Engine/Plugins/Editor/WaveformEditor/Source/WaveformEditorWidgets/Public/SWaveformTransformationsOverlay.h
--- a/Engine/Plugins/Editor/WaveformEditor/Source/WaveformEditorWidgets/Public/SWaveformTransformationsOverlay.h
+++ b/Engine/Plugins/Editor/WaveformEditor/Source/WaveformEditorWidgets/Public/SWaveformTransformationsOverlay.h
@@ -16,14 +16,20 @@ class WAVEFORMEDITORWIDGETS_API SWaveformTransformationsOverlay : public SCompou
 {
 public:
 
-	SLATE_BEGIN_ARGS(SWaveformTransformationsOverlay) {}
+	SLATE_BEGIN_ARGS(SWaveformTransformationsOverlay) 
+		: _AnchorsRatioConverter([](const float InRatio) {return InRatio; })
+	{
+	}
+
 		SLATE_DEFAULT_SLOT(FArguments, InArgs)
+
+		SLATE_ARGUMENT(TFunction<float(const float)>, AnchorsRatioConverter)
+
 	SLATE_END_ARGS()
 
-	void Construct(const FArguments& InArgs, TArrayView< const FTransformationLayerRenderInfo> InTransformationRenderers, TSharedRef<FSparseSampledSequenceTransportCoordinator> InTransportCoordinator);
+	void Construct(const FArguments& InArgs, TArrayView< const FTransformationLayerRenderInfo> InTransformationRenderers);
 	void OnLayerChainGenerated(FTransformationLayerRenderInfo* FirstLayerPtr, const int32 NLayers);
 	void UpdateLayerConstraints();
-	void OnNewWaveformDisplayRange(const TRange<double> NewDisplayRange);
 
 	FReply OnMouseButtonDown(const FGeometry& MyGeometry, const FPointerEvent& MouseEvent) override;
 	FReply OnMouseButtonUp(const FGeometry& MyGeometry, const FPointerEvent& MouseEvent) override;
@@ -39,5 +45,5 @@ private:
 	TArray<TSharedPtr<SWidget>> TransformationLayers;
 	TArray<SConstraintCanvas::FSlot*> LayersSlots;
 	TArrayView<const FTransformationLayerRenderInfo> TransformationRenderers;
-	TSharedPtr<FSparseSampledSequenceTransportCoordinator> TransportCoordinator = nullptr;
+	TFunction<float(const float)> AnchorsRatioConverter;
 };
\ No newline at end of file
