[WaveformEditor] Init size of TransformedPCMData when copying a transformations buffer in a chain that does not change the file lenght.
Fix for #jira UE-185694
#trivial

diff --git a/Engine/Plugins/Editor/WaveformEditor/Source/WaveformEditorWidgets/Private/WaveformEditorSequenceDataProvider.cpp b/Engine/Plugins/Editor/WaveformEditor/Source/WaveformEditorWidgets/Private/WaveformEditorSequenceDataProvider.cpp
--- a/Engine/Plugins/Editor/WaveformEditor/Source/WaveformEditorWidgets/Private/WaveformEditorSequenceDataProvider.cpp
+++ b/Engine/Plugins/Editor/WaveformEditor/Source/WaveformEditorWidgets/Private/WaveformEditorSequenceDataProvider.cpp
@@ -163,6 +163,10 @@ void FWaveformEditorSequenceDataProvider::GenerateSequenceDataInternal()
 
 		if (!bChainChangesFileLength)
 		{
+			if (TransformedPCMData.Num() != TransformationsBuffer.Num())
+			{
+				TransformedPCMData.SetNumUninitialized(TransformationsBuffer.Num());
+			}
 			FMemory::Memcpy(TransformedPCMData.GetData(), TransformationsBuffer.GetData(), TransformationsBuffer.GetAllocatedSize());
 		}
 		else
