Improved Transport Controls
Waveform editor Stop control now stops to the last cached play location: for natural and explicit stops for one-shot and looped sounds.
Also verified no regression for Play-Pause/Stop-Play cases.
FWaveformEditorTransportController now manages/responds to more of the underlying AudioComponent state itself, rather than relying on the caller.

diff --git a/Engine/Plugins/Editor/WaveformEditor/Source/WaveformEditor/Private/WaveformEditor.cpp b/Engine/Plugins/Editor/WaveformEditor/Source/WaveformEditor/Private/WaveformEditor.cpp
--- a/Engine/Plugins/Editor/WaveformEditor/Source/WaveformEditor/Private/WaveformEditor.cpp
+++ b/Engine/Plugins/Editor/WaveformEditor/Source/WaveformEditor/Private/WaveformEditor.cpp
@@ -840,8 +840,8 @@ void FWaveformEditor::HandleAudioComponentPlayStateChanged(const UAudioComponent
 
 		if (!bIgnoreStopState)
 		{
-			TransportController->CacheStartTime(0.f);
-			TransportCoordinator->Stop();
+			const float CachedAudioStartTimeAsPercentage = TransportController->GetCachedAudioStartTimeAsPercentage();
+			TransportCoordinator->SetProgressRatio(CachedAudioStartTimeAsPercentage); // show cached start time for next play rather than playhead stop point (doesn't affect pause)
 		}
 
 		break;
diff --git a/Engine/Plugins/Editor/WaveformEditor/Source/WaveformEditor/Private/WaveformEditorTransportController.cpp b/Engine/Plugins/Editor/WaveformEditor/Source/WaveformEditor/Private/WaveformEditorTransportController.cpp
--- a/Engine/Plugins/Editor/WaveformEditor/Source/WaveformEditor/Private/WaveformEditorTransportController.cpp
+++ b/Engine/Plugins/Editor/WaveformEditor/Source/WaveformEditor/Private/WaveformEditorTransportController.cpp
@@ -1,14 +1,20 @@
 // Copyright Epic Games, Inc. All Rights Reserved.
 
 #include "WaveformEditorTransportController.h"
+#include "Sound/SoundBase.h"
 
 FWaveformEditorTransportController::FWaveformEditorTransportController(UAudioComponent* InAudioComponent)
 	: AudioComponent(InAudioComponent)
-{}
+{
+	check(InAudioComponent != nullptr);
+	AudioComponentPlayStateChangedDelegateHandle = InAudioComponent->OnAudioPlayStateChangedNative.AddRaw(this, &FWaveformEditorTransportController::HandleAudioComponentPlayStateChanged);
+}
 
 FWaveformEditorTransportController::~FWaveformEditorTransportController()
 {
 	Stop();
+	check(AudioComponent != nullptr);
+	AudioComponent->OnAudioPlayStateChangedNative.Remove(AudioComponentPlayStateChangedDelegateHandle);
 }
 
 void FWaveformEditorTransportController::Play()
@@ -29,8 +35,7 @@ void FWaveformEditorTransportController::Play()
 		
 	}
 
-	AudioComponent->Play(CachedAudioStartTime);
-	
+	AudioComponent->Play(CachedAudioStartTime);	
 }
 
 void FWaveformEditorTransportController::Play(const float StartTime)
@@ -62,8 +67,7 @@ void FWaveformEditorTransportController::Stop()
 		return;
 	}
 
-	Play(0.f);
-	AudioComponent->StopDelayed(0.1f);
+	AudioComponent->Stop();
 
 	if (IsPaused())
 	{
@@ -107,6 +111,7 @@ bool FWaveformEditorTransportController::IsPlaying() const
 
 void FWaveformEditorTransportController::CacheStartTime(const float StartTime)
 {
+	ensure(StartTime >= 0.0f);
 	CachedAudioStartTime = StartTime;
 
 	if (IsPaused())
@@ -124,7 +129,39 @@ void FWaveformEditorTransportController::Seek(const float SeekTime)
 	AudioComponent->Play(SeekTime);
 }
 
+float FWaveformEditorTransportController::GetCachedAudioStartTimeAsPercentage() const
+{
+	check(AudioComponent != nullptr);
+
+	if (AudioComponent->Sound != nullptr)
+	{
+		const float SoundDuration = AudioComponent->Sound->Duration; // not GetDuration() to get the raw duration and not INDEFINITELY_LOOPING_DURATION if looping
+		if (SoundDuration > 0.0f)
+		{
+			ensure(CachedAudioStartTime >= 0.0f);
+			const float PlaybackPercentage = CachedAudioStartTime / SoundDuration;
+			return FGenericPlatformMath::Fmod(PlaybackPercentage, 1.0f);
+		}
+	}
+	
+	return 0.0f;
+}
+
 const bool FWaveformEditorTransportController::SoundBaseIsValid() const
 {
 	return AudioComponent->GetSound() != nullptr;
-}
\ No newline at end of file
+}
+
+void FWaveformEditorTransportController::HandleAudioComponentPlayStateChanged(const UAudioComponent* InAudioComponent, EAudioComponentPlayState NewPlayState)
+{
+	ensure(InAudioComponent == AudioComponent);
+
+	switch (NewPlayState)
+	{
+	default:
+		break;
+	case EAudioComponentPlayState::Stopped:
+		bCachedTimeDuringPause = false;
+		break;
+	}
+}
diff --git a/Engine/Plugins/Editor/WaveformEditor/Source/WaveformEditor/Public/WaveformEditorTransportController.h b/Engine/Plugins/Editor/WaveformEditor/Source/WaveformEditor/Public/WaveformEditorTransportController.h
--- a/Engine/Plugins/Editor/WaveformEditor/Source/WaveformEditor/Public/WaveformEditorTransportController.h
+++ b/Engine/Plugins/Editor/WaveformEditor/Source/WaveformEditor/Public/WaveformEditorTransportController.h
@@ -3,6 +3,7 @@
 #pragma once
 
 #include "Components/AudioComponent.h"
+#include "Delegates/IDelegateInstance.h"
 
 class FWaveformEditorTransportController
 {
@@ -21,13 +22,15 @@ public:
 	bool IsPlaying() const;
 	void CacheStartTime(const float StartTime);
 	void Seek(const float SeekTime);
+	float GetCachedAudioStartTimeAsPercentage() const;
 
 private:
 	const bool SoundBaseIsValid() const;
+	// Playback delegates handlers
+	void HandleAudioComponentPlayStateChanged(const UAudioComponent* InAudioComponent, EAudioComponentPlayState NewPlayState);
+	FDelegateHandle AudioComponentPlayStateChangedDelegateHandle;
 
 	UAudioComponent* AudioComponent = nullptr;
 	float CachedAudioStartTime = 0.f;
 	bool bCachedTimeDuringPause = false;
-
-	FDelegateHandle PlayStateChangeDelegateHandle;
 };
\ No newline at end of file
diff --git a/Engine/Plugins/Runtime/AudioWidgets/Source/AudioWidgets/Private/SparseSampledSequenceTransportCoordinator.cpp b/Engine/Plugins/Runtime/AudioWidgets/Source/AudioWidgets/Private/SparseSampledSequenceTransportCoordinator.cpp
--- a/Engine/Plugins/Runtime/AudioWidgets/Source/AudioWidgets/Private/SparseSampledSequenceTransportCoordinator.cpp
+++ b/Engine/Plugins/Runtime/AudioWidgets/Source/AudioWidgets/Private/SparseSampledSequenceTransportCoordinator.cpp
@@ -96,11 +96,6 @@ void FSparseSampledSequenceTransportCoordinator::UpdateZoomRatioAndDisplayRange(
 	MoveFocusPoint(NewPlayheadPosition);
 }
 
-void FSparseSampledSequenceTransportCoordinator::Stop()
-{
-	SetProgressRatio(0.f);
-}
-
 const float FSparseSampledSequenceTransportCoordinator::GetFocusPoint() const
 {
 	return FocusPoint;
diff --git a/Engine/Plugins/Runtime/AudioWidgets/Source/AudioWidgets/Public/SparseSampledSequenceTransportCoordinator.h b/Engine/Plugins/Runtime/AudioWidgets/Source/AudioWidgets/Public/SparseSampledSequenceTransportCoordinator.h
--- a/Engine/Plugins/Runtime/AudioWidgets/Source/AudioWidgets/Public/SparseSampledSequenceTransportCoordinator.h
+++ b/Engine/Plugins/Runtime/AudioWidgets/Source/AudioWidgets/Public/SparseSampledSequenceTransportCoordinator.h
@@ -34,7 +34,6 @@ public:
 	float ConvertAbsoluteRatioToZoomed(const float InAbsoluteRatio) const override;
 	float ConvertZoomedRatioToAbsolute(const float InZoomedRatio) const override;
 	void UpdatePlaybackRange(const TRange<double>& NewRange);
-	void Stop();
 
 private:
 	FORCEINLINE void MoveFocusPoint(const double InFocusPoint);
