Markers created as WaveformEditor Transformation

Plugin Notes:
-) WaveformTransformationMarkers - Retrieves and maintains the marker/loop region data to be displayed
-) WaveformTransformationMarkersRenderer - Renders the marker data over the waveform
-) WaveformTransformationsWidgetsSettings - Gives the user access to variables to alter the visuals of Transformations

Engine Changes:
-) IWaveformTransformation.h and SoundWave.cpp altered to update Transformations with the SoundWaveCuePoint data.

diff --git a/Engine/Plugins/Editor/WaveformEditor/Source/WaveformTransformations/Private/WaveformTransformationMarkers.cpp b/Engine/Plugins/Editor/WaveformEditor/Source/WaveformTransformations/Private/WaveformTransformationMarkers.cpp
new file mode 100644
--- /dev/null
+++ b/Engine/Plugins/Editor/WaveformEditor/Source/WaveformTransformations/Private/WaveformTransformationMarkers.cpp
@@ -0,0 +1,27 @@
+// Copyright Epic Games, Inc. All Rights Reserved.
+
+#include "WaveformTransformationMarkers.h"
+
+UWaveformTransformationMarkers::UWaveformTransformationMarkers(const FObjectInitializer& ObjectInitializer) : Super(ObjectInitializer)
+{
+	if(Markers == nullptr)
+	{
+		Markers = ObjectInitializer.CreateDefaultSubobject<UWaveCueArray>(this, TEXT("Markers"));
+	}
+}
+
+Audio::FTransformationPtr UWaveformTransformationMarkers::CreateTransformation() const
+{
+	return MakeUnique<FWaveTransformationMarkers>();
+}
+
+void UWaveformTransformationMarkers::UpdateConfiguration(FWaveTransformUObjectConfiguration& InOutConfiguration)
+{
+	check(Markers != nullptr);
+	
+	Markers->WaveCues.Empty();
+	for (const FWaveCue& WaveCue : InOutConfiguration.WaveCues)
+	{
+		Markers->WaveCues.Emplace(WaveCue);
+	}
+}
\ No newline at end of file
diff --git a/Engine/Plugins/Editor/WaveformEditor/Source/WaveformTransformations/Public/WaveformTransformationMarkers.h b/Engine/Plugins/Editor/WaveformEditor/Source/WaveformTransformations/Public/WaveformTransformationMarkers.h
new file mode 100644
--- /dev/null
+++ b/Engine/Plugins/Editor/WaveformEditor/Source/WaveformTransformations/Public/WaveformTransformationMarkers.h
@@ -0,0 +1,64 @@
+// Copyright Epic Games, Inc. All Rights Reserved.
+
+#pragma once
+
+#include "IWaveformTransformation.h"
+
+#include "WaveformTransformationMarkers.generated.h"
+
+//USTRUCT version of FWaveCue
+USTRUCT()
+struct FCuePoint
+{
+	GENERATED_BODY()
+	FCuePoint(const FWaveCue& WaveCue) 
+		: CuePointID(WaveCue.CuePointID)
+		, Position(WaveCue.Position)
+		, Label(WaveCue.Label)
+		, SampleLength(WaveCue.SampleLength){}
+	FCuePoint() {};
+
+	// Unique identifying value for the cue
+	UPROPERTY()
+	uint32 CuePointID = 0;
+	// Sample offset associated with the cue point
+	UPROPERTY()
+	uint32 Position = 0;
+	// Cue label
+	UPROPERTY()
+	FString Label = "";
+	// If this is a region, it will have a duration (sample length)
+	UPROPERTY()
+	uint32 SampleLength = 0;
+};
+
+//Used to make CuePoints work with PropertyHandles
+UCLASS(BlueprintType, DefaultToInstanced)
+class WAVEFORMTRANSFORMATIONS_API UWaveCueArray : public UObject
+{
+	GENERATED_BODY()
+public:
+	UPROPERTY()
+	TArray<FCuePoint> WaveCues;
+};
+
+class WAVEFORMTRANSFORMATIONS_API FWaveTransformationMarkers : public Audio::IWaveTransformation
+{
+public:
+	explicit FWaveTransformationMarkers() = default;
+};
+
+UCLASS()
+class WAVEFORMTRANSFORMATIONS_API UWaveformTransformationMarkers : public UWaveformTransformationBase
+{
+	GENERATED_BODY()
+public:
+	UWaveformTransformationMarkers(const FObjectInitializer& ObjectInitializer);
+
+	UPROPERTY(VisibleAnywhere, Category = "Markers")
+	TObjectPtr<UWaveCueArray> Markers;
+	
+	virtual Audio::FTransformationPtr CreateTransformation() const override;
+
+	void UpdateConfiguration(FWaveTransformUObjectConfiguration& InOutConfiguration) override;
+};
diff --git a/Engine/Plugins/Editor/WaveformEditor/Source/WaveformTransformationsWidgets/Private/WaveformTransformationMarkersRenderer.cpp b/Engine/Plugins/Editor/WaveformEditor/Source/WaveformTransformationsWidgets/Private/WaveformTransformationMarkersRenderer.cpp
new file mode 100644
--- /dev/null
+++ b/Engine/Plugins/Editor/WaveformEditor/Source/WaveformTransformationsWidgets/Private/WaveformTransformationMarkersRenderer.cpp
@@ -0,0 +1,246 @@
+// Copyright Epic Games, Inc. All Rights Reserved.
+
+#include "WaveformTransformationMarkersRenderer.h"
+
+#include "Brushes/SlateRoundedBoxBrush.h"
+#include "Layout/Geometry.h"
+#include "PropertyHandle.h"
+#include "Rendering/DrawElements.h"
+#include "Widgets/SWidget.h"
+
+FWaveformTransformationMarkerRenderer::FWaveformTransformationMarkerRenderer()
+{
+	const UWaveformTransformationsWidgetsSettings* Settings = GetDefault<UWaveformTransformationsWidgetsSettings>();
+	check(Settings);
+	Settings->OnSettingChanged().AddRaw(this, &FWaveformTransformationMarkerRenderer::OnWidgetSettingsUpdated);
+	ApplyWidgetSettings(Settings);
+}
+
+int32 FWaveformTransformationMarkerRenderer::OnPaint(const FPaintArgs& Args, const FGeometry& AllottedGeometry, const FSlateRect& MyCullingRect, FSlateWindowElementList& OutDrawElements, int32 LayerId, const FWidgetStyle& InWidgetStyle, bool bParentEnabled) const
+{
+	LayerId = DrawMarkerHandles(AllottedGeometry, OutDrawElements, LayerId);
+
+	return LayerId;
+}
+
+int32 FWaveformTransformationMarkerRenderer::DrawMarkerHandles(const FGeometry& AllottedGeometry, FSlateWindowElementList& OutDrawElements, int32 LayerId) const
+{
+	TArray<FVector2D> LinePoints;
+	LinePoints.SetNumUninitialized(2);
+
+	check(TransformationWaveInfo.NumChannels > 0);
+	check(TransformationWaveInfo.StartFrameOffset % TransformationWaveInfo.NumChannels == 0);
+	const uint32 StartFrameOffset = TransformationWaveInfo.StartFrameOffset / TransformationWaveInfo.NumChannels;
+	const float WindowWidth = AllottedGeometry.Size.X;
+	const float MarkerHeight = 30.f;
+	const float MarkerWidth = 8.f;
+	FSlateRoundedBoxBrush RoundedBoxBrush(FLinearColor::Red, 2);
+	FSlateFontInfo Font = FAppStyle::GetFontStyle("Regular");
+	Font.Size = LabelFontSize;
+	
+	uint32 ColorIndex = 0;
+
+	if (LoopPoints.Num() > 0)
+	{
+		check(LoopColors.Num() > 0);
+	}
+	for (const FCuePoint& LoopRegion : LoopPoints)
+	{
+		//Positions that become negative after StartFrameOffset subtraction are not rendered
+		const float LoopRegionPosition = (LoopRegion.Position - StartFrameOffset) * PixelsPerFrame;
+		LinePoints[0] = FVector2D(LoopRegionPosition, 0.f);
+		LinePoints[1] = FVector2D(LoopRegionPosition, AllottedGeometry.Size.Y);
+
+		ColorIndex = (ColorIndex + 1) % LoopColors.Num();
+
+		FSlateDrawElement::MakeLines(
+			OutDrawElements,
+			++LayerId,
+			AllottedGeometry.ToPaintGeometry(),
+			LinePoints,
+			ESlateDrawEffect::None,
+			LoopColors[ColorIndex],
+			false
+		);
+
+
+		const float MarkerCenter = LoopRegionPosition;
+		const float	HandleStart = MarkerCenter - MarkerWidth / 2;
+		const float	HandleEnd = HandleStart + MarkerWidth;
+		const float LoopRegionPixelLength = LoopRegion.SampleLength * PixelsPerFrame;
+		FSlateBrush Brush;
+		FPaintGeometry HandleGeometry = AllottedGeometry.ToPaintGeometry(FVector2f(MarkerWidth, MarkerHeight), FSlateLayoutTransform(FVector2f(HandleStart, 0)));
+		FPaintGeometry LoopBoxGeometry = AllottedGeometry.ToPaintGeometry(FVector2f(LoopRegionPixelLength, AllottedGeometry.GetLocalSize().Y), FSlateLayoutTransform(FVector2f(MarkerCenter, 0)));
+		LinePoints[0] = FVector2D(LoopRegionPosition + LoopRegionPixelLength, 0.f);
+		LinePoints[1] = FVector2D(LoopRegionPosition + LoopRegionPixelLength, AllottedGeometry.Size.Y);
+
+		FSlateDrawElement::MakeLines(
+			OutDrawElements,
+			++LayerId,
+			AllottedGeometry.ToPaintGeometry(),
+			LinePoints,
+			ESlateDrawEffect::None,
+			LoopColors[ColorIndex],
+			false
+		);
+		
+
+		FSlateDrawElement::MakeBox(
+			OutDrawElements,
+			++LayerId,
+			LoopBoxGeometry,
+			&Brush,
+			ESlateDrawEffect::InvertAlpha,
+			LoopColors[ColorIndex]
+		);
+
+		FSlateDrawElement::MakeBox(
+			OutDrawElements,
+			++LayerId,
+			HandleGeometry,
+			&RoundedBoxBrush,
+			ESlateDrawEffect::None,
+			LoopColors[ColorIndex]
+		);
+
+		FVector2D TextOffset(MarkerCenter + MarkerWidth, 0);
+		FSlateDrawElement::MakeText(
+			OutDrawElements,
+			++LayerId,
+			AllottedGeometry.ToPaintGeometry(AllottedGeometry.Size, FSlateLayoutTransform(TextOffset)),
+			LoopRegion.Label,
+			Font,
+			ESlateDrawEffect::None,
+			LabelTextColor
+		);
+	}
+
+	for (const FCuePoint& Marker : CuePoints)
+	{
+		//Positions that become negative after StartFrameOffset subtraction are not rendered
+		const float MarkerPosition = (Marker.Position - StartFrameOffset) * PixelsPerFrame;
+		LinePoints[0] = FVector2D(MarkerPosition, 0.f);
+		LinePoints[1] = FVector2D(MarkerPosition, AllottedGeometry.Size.Y);
+
+		FSlateDrawElement::MakeLines(
+			OutDrawElements,
+			++LayerId,
+			AllottedGeometry.ToPaintGeometry(),
+			LinePoints,
+			ESlateDrawEffect::None,
+			MarkerColor,
+			false
+		);
+
+		const float	HandleStart = MarkerPosition - MarkerWidth / 2;
+		const float	HandleEnd = HandleStart + MarkerWidth;
+		FPaintGeometry HandleGeometry = AllottedGeometry.ToPaintGeometry(FVector2f(MarkerWidth, MarkerHeight), FSlateLayoutTransform(FVector2f(HandleStart, 0)));
+
+		FSlateDrawElement::MakeBox(
+			OutDrawElements,
+			++LayerId,
+			HandleGeometry,
+			&RoundedBoxBrush,
+			ESlateDrawEffect::None,
+			MarkerColor
+		);
+
+		FVector2D TextOffset(MarkerPosition + MarkerWidth, 0);
+		FSlateDrawElement::MakeText(
+			OutDrawElements,
+			++LayerId,
+			AllottedGeometry.ToPaintGeometry(AllottedGeometry.Size, FSlateLayoutTransform(TextOffset)),
+			Marker.Label,
+			Font,
+			ESlateDrawEffect::None,
+			LabelTextColor
+		);
+	}
+
+	return LayerId;
+}
+
+void FWaveformTransformationMarkerRenderer::SetPropertyHandles(const TArray<TSharedRef<IPropertyHandle>>& InPropertyHandles)
+{
+	FWaveformTransformationRendererBase::SetPropertyHandles(InPropertyHandles);
+
+	MarkersHandle = GetPropertyHandle(GET_MEMBER_NAME_CHECKED(UWaveformTransformationMarkers, Markers));
+}
+
+void FWaveformTransformationMarkerRenderer::Tick(const FGeometry& AllottedGeometry, const double InCurrentTime, const float InDeltaTime)
+{ 
+	check(TransformationWaveInfo.NumChannels > 0);
+	const float NumFrames = TransformationWaveInfo.NumAvilableSamples / TransformationWaveInfo.NumChannels;
+	check(NumFrames > 0);
+	PixelsPerFrame = FMath::Max(AllottedGeometry.GetLocalSize().X / NumFrames, 0.0);
+	ensure(PixelsPerFrame >= 0);
+
+	TMap<uint32, FCuePoint> CuePointMap;
+
+	MarkersArray = Cast<UWaveCueArray>(GetPropertyValue<UObject*>(MarkersHandle));
+	if (MarkersArray == nullptr)
+	{
+		return;
+	}
+	CuePointMap.Reserve(MarkersArray->WaveCues.Num());
+
+	for (const FCuePoint& Marker : MarkersArray->WaveCues)
+	{
+		if (!CuePointMap.Contains(Marker.CuePointID))
+		{
+			CuePointMap.Emplace(Marker.CuePointID, Marker);
+		}
+		else if (Marker.SampleLength > 0)
+		{
+			CuePointMap[Marker.CuePointID].SampleLength = Marker.SampleLength;
+		}
+	}
+
+	CuePoints.Empty();
+	LoopPoints.Empty();
+	for (TPair<uint32, FCuePoint>& Pair : CuePointMap)
+	{
+		if (Pair.Value.SampleLength > 0)
+		{
+			LoopPoints.Add(MoveTemp(Pair.Value));
+		}
+		else
+		{
+			CuePoints.Add(MoveTemp(Pair.Value));
+		}
+	}
+	CuePointMap.Empty();
+}
+
+void FWaveformTransformationMarkerRenderer::OnWidgetSettingsUpdated(const FName& PropertyName, const UWaveformTransformationsWidgetsSettings* Settings)
+{
+	check(Settings)
+
+	if (PropertyName == GET_MEMBER_NAME_CHECKED(UWaveformTransformationsWidgetsSettings, MarkerColor))
+	{
+		MarkerColor = Settings->MarkerColor;
+	}
+	if (PropertyName == GET_MEMBER_NAME_CHECKED(UWaveformTransformationsWidgetsSettings, LoopColors))
+	{
+		LoopColors = Settings->LoopColors;
+	}
+	if (PropertyName == GET_MEMBER_NAME_CHECKED(UWaveformTransformationsWidgetsSettings, LabelTextColor))
+	{
+		LabelTextColor = Settings->LabelTextColor;
+	}
+	if (PropertyName == GET_MEMBER_NAME_CHECKED(UWaveformTransformationsWidgetsSettings, LabelFontSize))
+	{
+		LabelFontSize = Settings->LabelFontSize;
+	}
+}
+
+void FWaveformTransformationMarkerRenderer::ApplyWidgetSettings(const UWaveformTransformationsWidgetsSettings* Settings)
+{
+	check(Settings)
+
+	MarkerColor = Settings->MarkerColor;
+	LoopColors = Settings->LoopColors;
+	LabelTextColor = Settings->LabelTextColor;
+	LabelFontSize = Settings->LabelFontSize;
+}
+
diff --git a/Engine/Plugins/Editor/WaveformEditor/Source/WaveformTransformationsWidgets/Private/WaveformTransformationMarkersRenderer.h b/Engine/Plugins/Editor/WaveformEditor/Source/WaveformTransformationsWidgets/Private/WaveformTransformationMarkersRenderer.h
new file mode 100644
--- /dev/null
+++ b/Engine/Plugins/Editor/WaveformEditor/Source/WaveformTransformationsWidgets/Private/WaveformTransformationMarkersRenderer.h
@@ -0,0 +1,39 @@
+// Copyright Epic Games, Inc. All Rights Reserved.
+
+#pragma once
+
+#include "WaveformTransformationMarkers.h"
+#include "WaveformTransformationRendererBase.h"
+#include "WaveformTransformationsWidgetsSettings.h"
+
+class FWaveformTransformationMarkerRenderer : public FWaveformTransformationRendererBase
+{
+public:
+	FWaveformTransformationMarkerRenderer();
+
+	virtual int32 OnPaint(const FPaintArgs& Args, const FGeometry& AllottedGeometry, const FSlateRect& MyCullingRect, FSlateWindowElementList& OutDrawElements, int32 LayerId, const FWidgetStyle& InWidgetStyle, bool bParentEnabled) const override;
+	virtual void Tick(const FGeometry& AllottedGeometry, const double InCurrentTime, const float InDeltaTime) override;
+
+	virtual void SetPropertyHandles(const TArray<TSharedRef<IPropertyHandle>>& InPropertyHandles) override;
+
+private:
+	int32 DrawMarkerHandles(const FGeometry& AllottedGeometry, FSlateWindowElementList& OutDrawElements, int32 LayerId) const;
+
+	void OnWidgetSettingsUpdated(const FName& PropertyName, const UWaveformTransformationsWidgetsSettings* Settings);
+	void ApplyWidgetSettings(const UWaveformTransformationsWidgetsSettings* Settings);
+
+	TArray<FCuePoint> CuePoints;
+	TArray<FCuePoint> LoopPoints;
+
+	TObjectPtr<UWaveCueArray> MarkersArray;
+
+	double PixelsPerFrame = 0.0;
+
+	FLinearColor MarkerColor = WaveformTransformationWidgetSharedDefaults::DefaultMarkerColor;
+	TArray<FLinearColor> LoopColors = WaveformTransformationWidgetSharedDefaults::DefaultLoopColors;
+	FLinearColor LabelTextColor = WaveformTransformationWidgetSharedDefaults::DefaultLabelTextColor;
+	float LabelFontSize = WaveformTransformationWidgetSharedDefaults::DefaultLabelFontSize;
+
+	TSharedPtr<IPropertyHandle> MarkersHandle;
+
+};
diff --git a/Engine/Plugins/Editor/WaveformEditor/Source/WaveformTransformationsWidgets/Private/WaveformTransformationsWidgetsModule.cpp b/Engine/Plugins/Editor/WaveformEditor/Source/WaveformTransformationsWidgets/Private/WaveformTransformationsWidgetsModule.cpp
--- a/Engine/Plugins/Editor/WaveformEditor/Source/WaveformTransformationsWidgets/Private/WaveformTransformationsWidgetsModule.cpp
+++ b/Engine/Plugins/Editor/WaveformEditor/Source/WaveformTransformationsWidgets/Private/WaveformTransformationsWidgetsModule.cpp
@@ -6,11 +6,14 @@
 #include "WaveformTransformationRendererMapper.h"
 #include "WaveformTransformationTrimFade.h"
 #include "WaveformTransformationTrimFadeRenderer.h"
+#include "WaveformTransformationMarkers.h"
+#include "WaveformTransformationMarkersRenderer.h"
 
 void FWaveformTransformationsWidgetsModule::StartupModule()
 {
 	FWaveformTransformationRendererMapper& RendererMapper = FWaveformTransformationRendererMapper::Get();
 	RendererMapper.RegisterRenderer<FWaveformTransformationTrimFadeRenderer>(UWaveformTransformationTrimFade::StaticClass());
+	RendererMapper.RegisterRenderer<FWaveformTransformationMarkerRenderer>(UWaveformTransformationMarkers::StaticClass());
 }
 
 void FWaveformTransformationsWidgetsModule::ShutdownModule()
diff --git a/Engine/Plugins/Editor/WaveformEditor/Source/WaveformTransformationsWidgets/Private/WaveformTransformationsWidgetsSettings.cpp b/Engine/Plugins/Editor/WaveformEditor/Source/WaveformTransformationsWidgets/Private/WaveformTransformationsWidgetsSettings.cpp
new file mode 100644
--- /dev/null
+++ b/Engine/Plugins/Editor/WaveformEditor/Source/WaveformTransformationsWidgets/Private/WaveformTransformationsWidgetsSettings.cpp
@@ -0,0 +1,46 @@
+// Copyright Epic Games, Inc. All Rights Reserved.
+
+#include "WaveformTransformationsWidgetsSettings.h"
+
+UWaveformTransformationsWidgetsSettings::UWaveformTransformationsWidgetsSettings(const FObjectInitializer& ObjectInitializer)
+	: Super(ObjectInitializer)
+	, MarkerColor(WaveformTransformationWidgetSharedDefaults::DefaultMarkerColor)
+	, LoopColors(WaveformTransformationWidgetSharedDefaults::DefaultLoopColors)
+	, LabelTextColor(WaveformTransformationWidgetSharedDefaults::DefaultLabelTextColor)
+	, LabelFontSize(WaveformTransformationWidgetSharedDefaults::DefaultLabelFontSize)
+{
+}
+
+FName UWaveformTransformationsWidgetsSettings::GetCategoryName() const
+{
+	return TEXT("Plugins");
+}
+
+FName UWaveformTransformationsWidgetsSettings::GetSectionName() const
+{
+	return TEXT("Waveform Transformations Display");
+}
+
+#if WITH_EDITOR
+FText UWaveformTransformationsWidgetsSettings::GetSectionText() const
+{
+	return NSLOCTEXT("WaveformTransformationsDisplay", "WaveformTransformationsDisplaySettingsSection", "Waveform Transformations Display");
+}
+
+void UWaveformTransformationsWidgetsSettings::PostEditChangeProperty(FPropertyChangedEvent& PropertyChangedEvent)
+{
+	Super::PostEditChangeProperty(PropertyChangedEvent);
+
+	if (PropertyChangedEvent.Property != nullptr)
+	{
+		SettingsChangedDelegate.Broadcast(PropertyChangedEvent.GetPropertyName(), this);
+	}
+}
+
+FOnWaveformTransformationsWidgetsSettingsChanged& UWaveformTransformationsWidgetsSettings::OnSettingChanged()
+{
+	return SettingsChangedDelegate;
+}
+
+FOnWaveformTransformationsWidgetsSettingsChanged UWaveformTransformationsWidgetsSettings::SettingsChangedDelegate;
+#endif// WITH_EDITOR
\ No newline at end of file
diff --git a/Engine/Plugins/Editor/WaveformEditor/Source/WaveformTransformationsWidgets/Private/WaveformTransformationsWidgetsSettings.h b/Engine/Plugins/Editor/WaveformEditor/Source/WaveformTransformationsWidgets/Private/WaveformTransformationsWidgetsSettings.h
new file mode 100644
--- /dev/null
+++ b/Engine/Plugins/Editor/WaveformEditor/Source/WaveformTransformationsWidgets/Private/WaveformTransformationsWidgetsSettings.h
@@ -0,0 +1,60 @@
+// Copyright Epic Games, Inc. All Rights Reserved.
+
+#pragma once
+
+#include "Engine/DeveloperSettings.h"
+
+#include "WaveformTransformationsWidgetsSettings.generated.h"
+
+DECLARE_MULTICAST_DELEGATE_TwoParams(FOnWaveformTransformationsWidgetsSettingsChanged, const FName& /*Property Name*/, const UWaveformTransformationsWidgetsSettings*);
+
+namespace WaveformTransformationWidgetSharedDefaults
+{
+	const FLinearColor DefaultMarkerColor			= FLinearColor(1.f, 0.9f, 0, 0.75f);
+	const TArray<FLinearColor> DefaultLoopColors	= { FLinearColor(0.8f, 0.6f, 1.f, 0.75f), FLinearColor(0.2f, 0.7f, 1.f, 0.75f) };
+	const FLinearColor DefaultLabelTextColor		= FLinearColor(0.9f, 0.9f, 0.9, 1.f);
+	const float DefaultLabelFontSize				= 10.f;
+}
+
+UCLASS(config = EditorPerProjectUserSettings, defaultconfig, meta = (DisplayName = "Waveform Transformations Display"))
+class UWaveformTransformationsWidgetsSettings : public UDeveloperSettings
+{
+	GENERATED_BODY()
+
+public:
+
+	UWaveformTransformationsWidgetsSettings(const FObjectInitializer& ObjectInitializer);
+
+	//~ Begin UDeveloperSettings interface
+	virtual FName GetCategoryName() const;
+
+	virtual FName GetSectionName() const override;
+
+#if WITH_EDITOR
+	virtual FText GetSectionText() const override;
+
+	virtual void PostEditChangeProperty(FPropertyChangedEvent& PropertyChangedEvent) override;
+	//~ End UDeveloperSettings interface
+
+	static FOnWaveformTransformationsWidgetsSettingsChanged& OnSettingChanged();
+#endif
+
+
+	//Set the color of your Waveform Markers
+	UPROPERTY(config, EditAnywhere, Category = "Markers")
+	FLinearColor MarkerColor;
+
+	//A list of colors so marker regions can have different colors
+	UPROPERTY(config, EditAnywhere, Category = "Markers")
+	TArray<FLinearColor> LoopColors;
+
+	UPROPERTY(config, EditAnywhere, Category = "Markers")
+	FLinearColor LabelTextColor;
+
+	UPROPERTY(config, EditAnywhere, Category = "Markers", Meta = (ClampMin = "1", ClampMax = "15"))
+	float LabelFontSize;
+
+
+private: 
+	static FOnWaveformTransformationsWidgetsSettingsChanged SettingsChangedDelegate;
+};
\ No newline at end of file
diff --git a/Engine/Source/Runtime/AudioExtensions/Public/IWaveformTransformation.h b/Engine/Source/Runtime/AudioExtensions/Public/IWaveformTransformation.h
--- a/Engine/Source/Runtime/AudioExtensions/Public/IWaveformTransformation.h
+++ b/Engine/Source/Runtime/AudioExtensions/Public/IWaveformTransformation.h
@@ -2,6 +2,7 @@
 
 #pragma once
 
+#include "Audio.h"
 #include "Containers/Array.h"
 #include "HAL/Platform.h"
 #include "DSP/BufferVectorOperations.h"
@@ -56,6 +57,7 @@ struct FWaveTransformUObjectConfiguration
 	float SampleRate = 0;
 	float StartTime = 0.f; 
 	float EndTime = -1.f; 
+	TArray<FWaveCue> WaveCues; // List of cues parsed from the wave file
 };
 
 // Base class to hold editor configurable properties for an arbitrary transformation of audio waveform data
diff --git a/Engine/Source/Runtime/Engine/Private/SoundWave.cpp b/Engine/Source/Runtime/Engine/Private/SoundWave.cpp
--- a/Engine/Source/Runtime/Engine/Private/SoundWave.cpp
+++ b/Engine/Source/Runtime/Engine/Private/SoundWave.cpp
@@ -4549,6 +4549,16 @@ const FWaveTransformUObjectConfiguration& USoundWave::UpdateTransformations()
 	TransformationChainConfig.SampleRate = ImportedSampleRate > 0 ? ImportedSampleRate : GetSampleRateForCurrentPlatform();
 	TransformationChainConfig.StartTime = 0.f;
 	TransformationChainConfig.EndTime = TotalSamples / TransformationChainConfig.SampleRate;
+	for (const FSoundWaveCuePoint& CuePoint : CuePoints)
+	{
+		FWaveCue WaveCue;
+		WaveCue.CuePointID = CuePoint.CuePointID;
+		WaveCue.Position = CuePoint.FramePosition;
+		WaveCue.Label = CuePoint.Label;
+		WaveCue.SampleLength = CuePoint.FrameLength;
+
+		TransformationChainConfig.WaveCues.Emplace(WaveCue);
+	}
 
 	for (TObjectPtr<UWaveformTransformationBase>& Transformation : Transformations)
 	{
