[Backout] - CL38392150
Original CL Desc
-----------------------------------------------------------------
Markers created as WaveformEditor Transformation

Plugin Notes:
-) WaveformTransformationMarkers - Retrieves and maintains the marker/loop region data to be displayed
-) WaveformTransformationMarkersRenderer - Renders the marker data over the waveform
-) WaveformTransformationsWidgetsSettings - Gives the user access to variables to alter the visuals of Transformations

Engine Changes:
-) IWaveformTransformation.h and SoundWave.cpp altered to update Transformations with the SoundWaveCuePoint data.

diff --git a/Engine/Plugins/Editor/WaveformEditor/Source/WaveformTransformations/Private/WaveformTransformationMarkers.cpp b/Engine/Plugins/Editor/WaveformEditor/Source/WaveformTransformations/Private/WaveformTransformationMarkers.cpp
deleted file mode 100644
diff --git a/Engine/Plugins/Editor/WaveformEditor/Source/WaveformTransformations/Public/WaveformTransformationMarkers.h b/Engine/Plugins/Editor/WaveformEditor/Source/WaveformTransformations/Public/WaveformTransformationMarkers.h
deleted file mode 100644
diff --git a/Engine/Plugins/Editor/WaveformEditor/Source/WaveformTransformationsWidgets/Private/WaveformTransformationMarkersRenderer.cpp b/Engine/Plugins/Editor/WaveformEditor/Source/WaveformTransformationsWidgets/Private/WaveformTransformationMarkersRenderer.cpp
deleted file mode 100644
diff --git a/Engine/Plugins/Editor/WaveformEditor/Source/WaveformTransformationsWidgets/Private/WaveformTransformationMarkersRenderer.h b/Engine/Plugins/Editor/WaveformEditor/Source/WaveformTransformationsWidgets/Private/WaveformTransformationMarkersRenderer.h
deleted file mode 100644
diff --git a/Engine/Plugins/Editor/WaveformEditor/Source/WaveformTransformationsWidgets/Private/WaveformTransformationsWidgetsModule.cpp b/Engine/Plugins/Editor/WaveformEditor/Source/WaveformTransformationsWidgets/Private/WaveformTransformationsWidgetsModule.cpp
--- a/Engine/Plugins/Editor/WaveformEditor/Source/WaveformTransformationsWidgets/Private/WaveformTransformationsWidgetsModule.cpp
+++ b/Engine/Plugins/Editor/WaveformEditor/Source/WaveformTransformationsWidgets/Private/WaveformTransformationsWidgetsModule.cpp
@@ -6,14 +6,11 @@
 #include "WaveformTransformationRendererMapper.h"
 #include "WaveformTransformationTrimFade.h"
 #include "WaveformTransformationTrimFadeRenderer.h"
-#include "WaveformTransformationMarkers.h"
-#include "WaveformTransformationMarkersRenderer.h"
 
 void FWaveformTransformationsWidgetsModule::StartupModule()
 {
 	FWaveformTransformationRendererMapper& RendererMapper = FWaveformTransformationRendererMapper::Get();
 	RendererMapper.RegisterRenderer<FWaveformTransformationTrimFadeRenderer>(UWaveformTransformationTrimFade::StaticClass());
-	RendererMapper.RegisterRenderer<FWaveformTransformationMarkerRenderer>(UWaveformTransformationMarkers::StaticClass());
 }
 
 void FWaveformTransformationsWidgetsModule::ShutdownModule()
diff --git a/Engine/Plugins/Editor/WaveformEditor/Source/WaveformTransformationsWidgets/Private/WaveformTransformationsWidgetsSettings.cpp b/Engine/Plugins/Editor/WaveformEditor/Source/WaveformTransformationsWidgets/Private/WaveformTransformationsWidgetsSettings.cpp
deleted file mode 100644
diff --git a/Engine/Plugins/Editor/WaveformEditor/Source/WaveformTransformationsWidgets/Private/WaveformTransformationsWidgetsSettings.h b/Engine/Plugins/Editor/WaveformEditor/Source/WaveformTransformationsWidgets/Private/WaveformTransformationsWidgetsSettings.h
deleted file mode 100644
diff --git a/Engine/Source/Runtime/AudioExtensions/Public/IWaveformTransformation.h b/Engine/Source/Runtime/AudioExtensions/Public/IWaveformTransformation.h
--- a/Engine/Source/Runtime/AudioExtensions/Public/IWaveformTransformation.h
+++ b/Engine/Source/Runtime/AudioExtensions/Public/IWaveformTransformation.h
@@ -2,7 +2,6 @@
 
 #pragma once
 
-#include "Audio.h"
 #include "Containers/Array.h"
 #include "HAL/Platform.h"
 #include "DSP/BufferVectorOperations.h"
@@ -57,7 +56,6 @@ struct FWaveTransformUObjectConfiguration
 	float SampleRate = 0;
 	float StartTime = 0.f; 
 	float EndTime = -1.f; 
-	TArray<FWaveCue> WaveCues; // List of cues parsed from the wave file
 };
 
 // Base class to hold editor configurable properties for an arbitrary transformation of audio waveform data
diff --git a/Engine/Source/Runtime/Engine/Private/SoundWave.cpp b/Engine/Source/Runtime/Engine/Private/SoundWave.cpp
--- a/Engine/Source/Runtime/Engine/Private/SoundWave.cpp
+++ b/Engine/Source/Runtime/Engine/Private/SoundWave.cpp
@@ -4549,16 +4549,6 @@ const FWaveTransformUObjectConfiguration& USoundWave::UpdateTransformations()
 	TransformationChainConfig.SampleRate = ImportedSampleRate > 0 ? ImportedSampleRate : GetSampleRateForCurrentPlatform();
 	TransformationChainConfig.StartTime = 0.f;
 	TransformationChainConfig.EndTime = TotalSamples / TransformationChainConfig.SampleRate;
-	for (const FSoundWaveCuePoint& CuePoint : CuePoints)
-	{
-		FWaveCue WaveCue;
-		WaveCue.CuePointID = CuePoint.CuePointID;
-		WaveCue.Position = CuePoint.FramePosition;
-		WaveCue.Label = CuePoint.Label;
-		WaveCue.SampleLength = CuePoint.FrameLength;
-
-		TransformationChainConfig.WaveCues.Emplace(WaveCue);
-	}
 
 	for (TObjectPtr<UWaveformTransformationBase>& Transformation : Transformations)
 	{
