Bug Fixes for:
[Waveform Editor] - Showing markers on a sound wave with loop regions when the loop regions colours array is empty crashes the editor
[Waveform Editor] - Alpha parameter values on Loop colours are backwards

diff --git a/Engine/Plugins/Editor/WaveformEditor/Source/WaveformTransformationsWidgets/Private/WaveformTransformationMarkersRenderer.cpp b/Engine/Plugins/Editor/WaveformEditor/Source/WaveformTransformationsWidgets/Private/WaveformTransformationMarkersRenderer.cpp
--- a/Engine/Plugins/Editor/WaveformEditor/Source/WaveformTransformationsWidgets/Private/WaveformTransformationMarkersRenderer.cpp
+++ b/Engine/Plugins/Editor/WaveformEditor/Source/WaveformTransformationsWidgets/Private/WaveformTransformationMarkersRenderer.cpp
@@ -84,14 +84,16 @@ int32 FWaveformTransformationMarkerRenderer::DrawMarkerHandles(const FGeometry&
 			false
 		);
 		
-
+		//Make LoopBox always half as opaque as the Marker
+		FLinearColor LoopBoxColor = LoopColors[ColorIndex];
+		LoopBoxColor.A *= 0.5;
 		FSlateDrawElement::MakeBox(
 			OutDrawElements,
 			++LayerId,
 			LoopBoxGeometry,
 			&Brush,
-			ESlateDrawEffect::InvertAlpha,
-			LoopColors[ColorIndex]
+			ESlateDrawEffect::None,
+			LoopBoxColor
 		);
 
 		FSlateDrawElement::MakeBox(
@@ -223,6 +225,10 @@ void FWaveformTransformationMarkerRenderer::OnWidgetSettingsUpdated(const FName&
 	if (PropertyName == GET_MEMBER_NAME_CHECKED(UWaveformTransformationsWidgetsSettings, LoopColors))
 	{
 		LoopColors = Settings->LoopColors;
+		if (LoopColors.Num() < 1)
+		{
+			LoopColors = WaveformTransformationWidgetSharedDefaults::DefaultLoopColors;
+		}
 	}
 	if (PropertyName == GET_MEMBER_NAME_CHECKED(UWaveformTransformationsWidgetsSettings, LabelTextColor))
 	{
