Bug Fix:
[Waveform Editor] - Showing markers on a sound wave with loop regions when the loop regions colours array is empty crashes the editor

diff --git a/Engine/Plugins/Editor/WaveformEditor/Source/WaveformTransformationsWidgets/Private/WaveformTransformationMarkersRenderer.cpp b/Engine/Plugins/Editor/WaveformEditor/Source/WaveformTransformationsWidgets/Private/WaveformTransformationMarkersRenderer.cpp
--- a/Engine/Plugins/Editor/WaveformEditor/Source/WaveformTransformationsWidgets/Private/WaveformTransformationMarkersRenderer.cpp
+++ b/Engine/Plugins/Editor/WaveformEditor/Source/WaveformTransformationsWidgets/Private/WaveformTransformationMarkersRenderer.cpp
@@ -224,11 +224,7 @@ void FWaveformTransformationMarkerRenderer::OnWidgetSettingsUpdated(const FName&
 	}
 	if (PropertyName == GET_MEMBER_NAME_CHECKED(UWaveformTransformationsWidgetsSettings, LoopColors))
 	{
-		LoopColors = Settings->LoopColors;
-		if (LoopColors.Num() < 1)
-		{
-			LoopColors = WaveformTransformationWidgetSharedDefaults::DefaultLoopColors;
-		}
+		SetLoopColors(Settings->LoopColors);
 	}
 	if (PropertyName == GET_MEMBER_NAME_CHECKED(UWaveformTransformationsWidgetsSettings, LabelTextColor))
 	{
@@ -245,8 +241,17 @@ void FWaveformTransformationMarkerRenderer::ApplyWidgetSettings(const UWaveformT
 	check(Settings)
 
 	MarkerColor = Settings->MarkerColor;
-	LoopColors = Settings->LoopColors;
+	SetLoopColors(Settings->LoopColors);
 	LabelTextColor = Settings->LabelTextColor;
 	LabelFontSize = Settings->LabelFontSize;
 }
 
+void FWaveformTransformationMarkerRenderer::SetLoopColors(const TArray<FLinearColor>& InColors)
+{
+	LoopColors = InColors;
+	if (LoopColors.Num() < 1)
+	{
+		LoopColors = WaveformTransformationWidgetSharedDefaults::DefaultLoopColors;
+	}
+}
+
diff --git a/Engine/Plugins/Editor/WaveformEditor/Source/WaveformTransformationsWidgets/Private/WaveformTransformationMarkersRenderer.h b/Engine/Plugins/Editor/WaveformEditor/Source/WaveformTransformationsWidgets/Private/WaveformTransformationMarkersRenderer.h
--- a/Engine/Plugins/Editor/WaveformEditor/Source/WaveformTransformationsWidgets/Private/WaveformTransformationMarkersRenderer.h
+++ b/Engine/Plugins/Editor/WaveformEditor/Source/WaveformTransformationsWidgets/Private/WaveformTransformationMarkersRenderer.h
@@ -22,6 +22,8 @@ private:
 	void OnWidgetSettingsUpdated(const FName& PropertyName, const UWaveformTransformationsWidgetsSettings* Settings);
 	void ApplyWidgetSettings(const UWaveformTransformationsWidgetsSettings* Settings);
 
+	void SetLoopColors(const TArray<FLinearColor>& InColors);
+
 	TArray<FSoundWaveCuePoint> CuePoints;
 	TArray<FSoundWaveCuePoint> LoopPoints;
 
