WaveEditor GUI - Add/Remove Cue Points/Markers
Added support to reimport and overwrite transformations

diff --git a/Engine/Plugins/Editor/WaveformEditor/Source/WaveformEditor/Private/WaveformEditor.cpp b/Engine/Plugins/Editor/WaveformEditor/Source/WaveformEditor/Private/WaveformEditor.cpp
--- a/Engine/Plugins/Editor/WaveformEditor/Source/WaveformEditor/Private/WaveformEditor.cpp
+++ b/Engine/Plugins/Editor/WaveformEditor/Source/WaveformEditor/Private/WaveformEditor.cpp
@@ -222,6 +222,16 @@ void FWaveformEditor::ExecuteReimport()
 	FReimportManager::Instance()->ValidateAllSourceFileAndReimport(CopyOfSelectedAssets, true, -1, bSelectNewAsset);
 }
 
+void FWaveformEditor::ExecuteReimportAndOverwriteTransformations()
+{
+	for (TObjectPtr<UWaveformTransformationBase> Transformation : SoundWave->Transformations)
+	{
+		Transformation->OverwriteTransformation();
+	}
+
+	ExecuteReimport();
+}
+
 void FWaveformEditor::RegisterTabSpawners(const TSharedRef<FTabManager>& InTabManager)
 {
 	WorkspaceMenuCategory = InTabManager->AddLocalWorkspaceMenuCategory(LOCTEXT("WorkspaceMenu_WaveformEditor", "Sound Wave Editor"));
@@ -425,6 +435,7 @@ TSharedRef<SWidget> FWaveformEditor::GenerateImportOptionsMenu()
 	MenuBuilder.BeginSection(NAME_None, LOCTEXT("ReimportMode_Label", "Reimport Mode"));
 	{
 		MenuBuilder.AddMenuEntry(FWaveformEditorCommands::Get().ReimportModeSameFile);
+		MenuBuilder.AddMenuEntry(FWaveformEditorCommands::Get().ReimportModeSameFileOverwriteTransformations);
 		MenuBuilder.AddMenuEntry(FWaveformEditorCommands::Get().ReimportModeNewFile);
 	}
 	MenuBuilder.EndSection();
@@ -503,6 +514,12 @@ bool FWaveformEditor::BindCommands()
 		FCanExecuteAction::CreateLambda([this] { return CanExecuteReimport(); }),
 		FIsActionChecked::CreateLambda([this] {return ReimportMode == EWaveEditorReimportMode::SameFile; }));
 
+	ToolkitCommands->MapAction(
+		Commands.ReimportModeSameFileOverwriteTransformations,
+		FExecuteAction::CreateLambda([this] { ReimportMode = EWaveEditorReimportMode::SameFile; ExecuteReimportAndOverwriteTransformations(); }),
+		FCanExecuteAction::CreateLambda([this] { return CanExecuteReimport(); }),
+		FIsActionChecked::CreateLambda([this] {return ReimportMode == EWaveEditorReimportMode::SameFile; }));
+
 	ToolkitCommands->MapAction(
 		Commands.ReimportModeNewFile,
 		FExecuteAction::CreateLambda([this] { ReimportMode = EWaveEditorReimportMode::SelectFile; ExecuteReimport(); }),
@@ -552,6 +569,10 @@ void FWaveformEditor::OnAssetReimport(UObject* ReimportedObject, bool bSuccessfu
 
 	if (ReimportedObject->IsA<USoundWave>() && ReimportedObject->GetPathName() == SoundWave->GetPathName())
 	{
+		//If the waveform editor is open, updating Transformations on reimport will show
+		// the effects of the overwritten transformations immediately instead of after a
+		// change is made by the user.
+		TransformationChainConfig = SoundWave->UpdateTransformations();
 		CreateWaveformView();
 		WaveformView.DataProvider->RequestSequenceView(TransportCoordinator->GetDisplayRange());
 		WaveformView.ViewWidget->SetPlayheadRatio(TransportCoordinator->GetFocusPoint());
diff --git a/Engine/Plugins/Editor/WaveformEditor/Source/WaveformEditor/Private/WaveformEditorCommands.cpp b/Engine/Plugins/Editor/WaveformEditor/Source/WaveformEditor/Private/WaveformEditorCommands.cpp
--- a/Engine/Plugins/Editor/WaveformEditor/Source/WaveformEditor/Private/WaveformEditorCommands.cpp
+++ b/Engine/Plugins/Editor/WaveformEditor/Source/WaveformEditor/Private/WaveformEditorCommands.cpp
@@ -29,6 +29,7 @@ void FWaveformEditorCommands::RegisterCommands()
 
 	UI_COMMAND(ReimportAsset, "Reimport Asset", "Reimports the inspected sound wave uasset", EUserInterfaceActionType::Button, FInputChord());
 	UI_COMMAND(ReimportModeSameFile, "Reimport from same file", "Reimports asset from the same file", EUserInterfaceActionType::Check, FInputChord());
+	UI_COMMAND(ReimportModeSameFileOverwriteTransformations, "Reimport from same file and overwrite transformations", "Reimports asset from the same file and overwrites the transformations", EUserInterfaceActionType::Check, FInputChord());
 	UI_COMMAND(ReimportModeNewFile, "Reimport from new file", "Reimports asset from a selected file", EUserInterfaceActionType::Check, FInputChord());
 }
 	
diff --git a/Engine/Plugins/Editor/WaveformEditor/Source/WaveformEditor/Private/WaveformEditorInstantiator.cpp b/Engine/Plugins/Editor/WaveformEditor/Source/WaveformEditor/Private/WaveformEditorInstantiator.cpp
--- a/Engine/Plugins/Editor/WaveformEditor/Source/WaveformEditor/Private/WaveformEditorInstantiator.cpp
+++ b/Engine/Plugins/Editor/WaveformEditor/Source/WaveformEditor/Private/WaveformEditorInstantiator.cpp
@@ -4,6 +4,7 @@
 
 #include "Algo/AnyOf.h"
 #include "ContentBrowserMenuContexts.h"
+#include "EditorReimportHandler.h"
 #include "Framework/Application/SlateApplication.h"
 #include "Framework/MultiBox/MultiBoxBuilder.h"
 #include "Sound/SoundWave.h"
@@ -42,6 +43,24 @@ void FWaveformEditorInstantiator::ExtendContentBrowserSelectionMenu()
 				}
 			}
 		}));
+	
+	FToolMenuSection& ImportSection = Menu->FindOrAddSection("ImportedAssetActions");
+
+	ImportSection.AddDynamicEntry("SoundWaveEditing_ReimportAndOverwriteTransformations", FNewToolMenuSectionDelegate::CreateLambda([this](FToolMenuSection& InSection)
+		{
+			if (const UContentBrowserAssetContextMenuContext* Context = InSection.FindContext<UContentBrowserAssetContextMenuContext>())
+			{
+				if (Algo::AnyOf(Context->SelectedAssets, WaveEditorInstantiatorPrivate::FilterUnwantedAssets))
+				{
+					const TAttribute<FText> Label = LOCTEXT("SoundWave_ReimportAndOverwriteTransformations", "Reimport And Overwrite Transformations");
+					const TAttribute<FText> ToolTip = LOCTEXT("SoundWave_ReimportAndOverwriteTransformations_Tooltip", "Reimport And Overwrite Transformations");
+					const FSlateIcon Icon = FSlateIcon(FAppStyle::GetAppStyleSetName(), "Icons.Import");
+					const FToolMenuExecuteAction UIAction = FToolMenuExecuteAction::CreateSP(this, &FWaveformEditorInstantiator::ExecuteReImportAndOverwriteTransformations);
+
+					InSection.AddMenuEntry("SoundWave_ReimportAndOverwriteTransformations", Label, ToolTip, Icon, UIAction);
+				}
+			}
+		}));
 }
 
 void FWaveformEditorInstantiator::ExecuteCreateWaveformEditor(const FToolMenuContext& MenuContext)
@@ -53,6 +72,29 @@ void FWaveformEditorInstantiator::ExecuteCreateWaveformEditor(const FToolMenuCon
 	}
 }
 
+void FWaveformEditorInstantiator::ExecuteReImportAndOverwriteTransformations(const FToolMenuContext& MenuContext)
+{
+	if (const UContentBrowserAssetContextMenuContext* Context = UContentBrowserAssetContextMenuContext::FindContextWithAssets(MenuContext))
+	{
+		TArray<USoundWave*> SoundWavesToEdit = Context->LoadSelectedObjectsIf<USoundWave>(WaveEditorInstantiatorPrivate::FilterUnwantedAssets);
+		TArray<UObject*> ObjectArray;
+
+		for (USoundWave* SoundWave : SoundWavesToEdit)
+		{
+			UObject* SelectedAsset = CastChecked<UObject>(SoundWave);
+			
+			ObjectArray.Emplace(SelectedAsset);
+
+			for (TObjectPtr<UWaveformTransformationBase> Transformation : SoundWave->Transformations)
+			{
+				Transformation->OverwriteTransformation();
+			}
+		}
+
+		FReimportManager::Instance()->ValidateAllSourceFileAndReimport(ObjectArray);
+	}
+}
+
 bool FWaveformEditorInstantiator::CanSoundWaveBeOpenedInEditor(const USoundWave* SoundWaveToEdit)
 {
 	bool bCanOpenWaveEditor = true;
diff --git a/Engine/Plugins/Editor/WaveformEditor/Source/WaveformEditor/Private/WaveformEditorInstantiator.h b/Engine/Plugins/Editor/WaveformEditor/Source/WaveformEditor/Private/WaveformEditorInstantiator.h
--- a/Engine/Plugins/Editor/WaveformEditor/Source/WaveformEditor/Private/WaveformEditorInstantiator.h
+++ b/Engine/Plugins/Editor/WaveformEditor/Source/WaveformEditor/Private/WaveformEditorInstantiator.h
@@ -17,6 +17,7 @@ public:
 private:
 	virtual void CreateWaveformEditor(TArray<USoundWave*> SoundWavesToEdit) override;
 	void ExecuteCreateWaveformEditor(const FToolMenuContext& MenuContext);
+	void ExecuteReImportAndOverwriteTransformations(const FToolMenuContext& MenuContext);
 	bool CanSoundWaveBeOpenedInEditor(const USoundWave* SoundWaveToEdit);
 	void DisplayErrorDialog(const FText& ErrorMessage) const;
 };
\ No newline at end of file
diff --git a/Engine/Plugins/Editor/WaveformEditor/Source/WaveformEditor/Public/WaveformEditor.h b/Engine/Plugins/Editor/WaveformEditor/Source/WaveformEditor/Public/WaveformEditor.h
--- a/Engine/Plugins/Editor/WaveformEditor/Source/WaveformEditor/Public/WaveformEditor.h
+++ b/Engine/Plugins/Editor/WaveformEditor/Source/WaveformEditor/Public/WaveformEditor.h
@@ -56,6 +56,7 @@ private:
 	bool BindDelegates();
 	bool SetUpAssetReimport();
 	void ExecuteReimport();
+	void ExecuteReimportAndOverwriteTransformations();
 
 	/**	Sets the wave editor layout */
 	const TSharedRef<FTabManager::FLayout> SetupStandaloneLayout();
diff --git a/Engine/Plugins/Editor/WaveformEditor/Source/WaveformEditor/Public/WaveformEditorCommands.h b/Engine/Plugins/Editor/WaveformEditor/Source/WaveformEditor/Public/WaveformEditorCommands.h
--- a/Engine/Plugins/Editor/WaveformEditor/Source/WaveformEditor/Public/WaveformEditorCommands.h
+++ b/Engine/Plugins/Editor/WaveformEditor/Source/WaveformEditor/Public/WaveformEditorCommands.h
@@ -29,6 +29,7 @@ public:
 
 	TSharedPtr<FUICommandInfo> ReimportAsset;
 	TSharedPtr<FUICommandInfo> ReimportModeSameFile;
+	TSharedPtr<FUICommandInfo> ReimportModeSameFileOverwriteTransformations;
 	TSharedPtr<FUICommandInfo> ReimportModeNewFile;
 
 };
diff --git a/Engine/Plugins/Editor/WaveformEditor/Source/WaveformTransformations/Private/WaveformTransformationMarkers.cpp b/Engine/Plugins/Editor/WaveformEditor/Source/WaveformTransformations/Private/WaveformTransformationMarkers.cpp
--- a/Engine/Plugins/Editor/WaveformEditor/Source/WaveformTransformations/Private/WaveformTransformationMarkers.cpp
+++ b/Engine/Plugins/Editor/WaveformEditor/Source/WaveformTransformations/Private/WaveformTransformationMarkers.cpp
@@ -18,6 +18,46 @@ Audio::FTransformationPtr UWaveformTransformationMarkers::CreateTransformation()
 void UWaveformTransformationMarkers::UpdateConfiguration(FWaveTransformUObjectConfiguration& InOutConfiguration)
 {
 	check(Markers != nullptr);
-	
-	Markers->WaveCues = InOutConfiguration.WaveCues;
-}
\ No newline at end of file
+
+	//Prevent USoundWave from overwriting the Transformation unintentionally
+	if(Markers->WaveCues.IsEmpty())
+	{
+		Markers->WaveCues = InOutConfiguration.WaveCues;
+	}
+}
+
+void UWaveformTransformationMarkers::OverwriteTransformation()
+{
+	Markers->WaveCues.Empty();
+}
+
+#if WITH_EDITOR
+void UWaveCueArray::PostEditChangeProperty(FPropertyChangedEvent& PropertyChangedEvent)
+{
+	Super::PostEditChangeProperty(PropertyChangedEvent);
+
+	const FName PropertyName = PropertyChangedEvent.GetPropertyName();
+	if (PropertyName == GET_MEMBER_NAME_CHECKED(UWaveCueArray, WaveCues) &&
+		PropertyChangedEvent.ChangeType == EPropertyChangeType::ArrayAdd)
+	{
+		check(WaveCues.Num() > 0);
+
+		int32 LargestCuePointID = 0;
+		for (int32 Index = WaveCues.Num()-1; Index >= 0; Index--)
+		{
+			if (WaveCues[Index].CuePointID > LargestCuePointID)
+			{
+				LargestCuePointID = WaveCues[Index].CuePointID + 1;
+			}
+		}
+		//Elements can be inserted at any index in the array
+		for (FSoundWaveCuePoint& CuePoint : WaveCues)
+		{
+			if (CuePoint.CuePointID == INDEX_NONE)
+			{
+				CuePoint.CuePointID = LargestCuePointID++;
+			}
+		}
+	}
+}
+#endif
\ No newline at end of file
diff --git a/Engine/Plugins/Editor/WaveformEditor/Source/WaveformTransformations/Public/WaveformTransformationMarkers.h b/Engine/Plugins/Editor/WaveformEditor/Source/WaveformTransformations/Public/WaveformTransformationMarkers.h
--- a/Engine/Plugins/Editor/WaveformEditor/Source/WaveformTransformations/Public/WaveformTransformationMarkers.h
+++ b/Engine/Plugins/Editor/WaveformEditor/Source/WaveformTransformations/Public/WaveformTransformationMarkers.h
@@ -12,7 +12,9 @@ class WAVEFORMTRANSFORMATIONS_API UWaveCueArray : public UObject
 {
 	GENERATED_BODY()
 public:
-	UPROPERTY()
+	void PostEditChangeProperty(FPropertyChangedEvent& PropertyChangedEvent);
+
+	UPROPERTY(EditAnywhere, Category = "Markers")
 	TArray<FSoundWaveCuePoint> WaveCues;
 };
 
@@ -35,4 +37,5 @@ public:
 	virtual Audio::FTransformationPtr CreateTransformation() const override;
 
 	void UpdateConfiguration(FWaveTransformUObjectConfiguration& InOutConfiguration) override;
+	void OverwriteTransformation() override;
 };
diff --git a/Engine/Plugins/Editor/WaveformEditor/Source/WaveformTransformationsWidgets/Private/WaveformTransformationMarkersRenderer.cpp b/Engine/Plugins/Editor/WaveformEditor/Source/WaveformTransformationsWidgets/Private/WaveformTransformationMarkersRenderer.cpp
--- a/Engine/Plugins/Editor/WaveformEditor/Source/WaveformTransformationsWidgets/Private/WaveformTransformationMarkersRenderer.cpp
+++ b/Engine/Plugins/Editor/WaveformEditor/Source/WaveformTransformationsWidgets/Private/WaveformTransformationMarkersRenderer.cpp
@@ -32,8 +32,8 @@ int32 FWaveformTransformationMarkerRenderer::DrawMarkerHandles(const FGeometry&
 	check(TransformationWaveInfo.StartFrameOffset % TransformationWaveInfo.NumChannels == 0);
 	const uint32 StartFrameOffset = TransformationWaveInfo.StartFrameOffset / TransformationWaveInfo.NumChannels;
 	const float WindowWidth = AllottedGeometry.Size.X;
-	const float MarkerHeight = 30.f;
-	const float MarkerWidth = 8.f;
+	const float MarkerHeight = AllottedGeometry.Size.Y * InteractionRatioYDelta;
+	const float MarkerWidth = InteractionPixelXDelta * 0.8;
 	FSlateRoundedBoxBrush RoundedBoxBrush(FLinearColor::Red, 2);
 	FSlateFontInfo Font = FAppStyle::GetFontStyle("Regular");
 	Font.Size = LabelFontSize;
@@ -186,8 +186,12 @@ void FWaveformTransformationMarkerRenderer::Tick(const FGeometry& AllottedGeomet
 	}
 	CuePointMap.Reserve(MarkersArray->WaveCues.Num());
 
+	bool IsMarkerInInteractionRangeValid = false;
 	for (const FSoundWaveCuePoint& Marker : MarkersArray->WaveCues)
 	{
+		//Fix duplicate loop region markers for SoundWaves that have not been reimported 
+		//after import fix
+		ensure(Marker.CuePointID != INDEX_NONE);
 		if (!CuePointMap.Contains(Marker.CuePointID))
 		{
 			CuePointMap.Emplace(Marker.CuePointID, Marker);
@@ -196,6 +200,17 @@ void FWaveformTransformationMarkerRenderer::Tick(const FGeometry& AllottedGeomet
 		{
 			CuePointMap[Marker.CuePointID].FrameLength = Marker.FrameLength;
 		}
+
+		if (&Marker == MarkerInInteractionRange.GetMarker())
+		{
+			IsMarkerInInteractionRangeValid = true;
+		}
+	}
+
+	//Invalidate Marker pointer if it has been removed from the array
+	if (!IsMarkerInInteractionRangeValid)
+	{
+		MarkerInInteractionRange.SetMarkerInInteractionRange(nullptr);
 	}
 
 	CuePoints.Empty();
@@ -214,6 +229,62 @@ void FWaveformTransformationMarkerRenderer::Tick(const FGeometry& AllottedGeomet
 	CuePointMap.Empty();
 }
 
+FReply FWaveformTransformationMarkerRenderer::OnMouseButtonDown(SWidget& OwnerWidget, const FGeometry& MyGeometry, const FPointerEvent& MouseEvent)
+{
+	const FVector2D LocalCursorPosition = GetLocalCursorPosition(MouseEvent, MyGeometry);
+
+	const bool bIsLeftMouseButton = MouseEvent.IsMouseButtonDown(EKeys::LeftMouseButton);
+
+	if (bIsLeftMouseButton)
+	{
+		MarkerInteractionType = GetInteractionTypeFromCursorPosition(LocalCursorPosition, MyGeometry);
+
+		if (MarkerInteractionType != EMarkerInteractionType::None)
+		{
+			BeginTransaction(TEXT("PropertyEditor"), FText::Format(NSLOCTEXT("PropertyEditor", "SetNumericPropertyTransaction", "Edit {0}"), FText::FromName(GET_MEMBER_NAME_CHECKED(UWaveformTransformationMarkers, Markers))), nullptr);
+			return FReply::Handled().CaptureMouse(OwnerWidget.AsShared()).PreventThrottling();
+		}
+
+	}
+
+	return FReply::Unhandled();
+}
+
+FReply FWaveformTransformationMarkerRenderer::OnMouseButtonUp(SWidget& OwnerWidget, const FGeometry& MyGeometry, const FPointerEvent& MouseEvent)
+{
+	if (MarkerInteractionType != EMarkerInteractionType::None)
+	{
+		SetPropertyValueDependingOnInteractionType(MouseEvent, MyGeometry);
+		EndTransaction();
+		MarkerInteractionType = EMarkerInteractionType::None;
+		return FReply::Handled().ReleaseMouseCapture();
+	}
+
+	return FReply::Unhandled();
+}
+
+FReply FWaveformTransformationMarkerRenderer::OnMouseMove(SWidget& OwnerWidget, const FGeometry& MyGeometry, const FPointerEvent& MouseEvent)
+{
+	if (MouseEvent.IsMouseButtonDown(EKeys::LeftMouseButton) && MarkerInteractionType != EMarkerInteractionType::None)
+	{
+		SetPropertyValueDependingOnInteractionType(MouseEvent, MyGeometry, EPropertyValueSetFlags::InteractiveChange | EPropertyValueSetFlags::NotTransactable);
+
+		return FReply::Handled().CaptureMouse(OwnerWidget.AsShared());
+	}
+
+	return FReply::Unhandled();
+}
+
+FCursorReply FWaveformTransformationMarkerRenderer::OnCursorQuery(const FGeometry& MyGeometry, const FPointerEvent& CursorEvent) const
+{
+	if (MarkerInteractionType == EMarkerInteractionType::ScrubbingMarkerHandle)
+	{
+		return FCursorReply::Cursor(EMouseCursor::GrabHandClosed);
+	}
+
+	return FCursorReply::Unhandled();
+}
+
 void FWaveformTransformationMarkerRenderer::OnWidgetSettingsUpdated(const FName& PropertyName, const UWaveformTransformationsWidgetsSettings* Settings)
 {
 	check(Settings)
@@ -255,3 +326,96 @@ void FWaveformTransformationMarkerRenderer::SetLoopColors(const TArray<FLinearCo
 	}
 }
 
+bool FWaveformTransformationMarkerRenderer::IsPositionInInteractionRange(const int32 InFramePosition, const FVector2D& InLocalCursorPosition, const FGeometry& WidgetGeometry) const
+{
+	TRange InteractionRange(TRange<float>::Inclusive(InFramePosition * PixelsPerFrame - InteractionPixelXDelta, InFramePosition * PixelsPerFrame + InteractionPixelXDelta));
+	if (InteractionRange.Contains(InLocalCursorPosition.X) &&
+		InLocalCursorPosition.Y < WidgetGeometry.GetLocalSize().Y * InteractionRatioYDelta)
+	{
+		return true;
+	}
+	return false;
+}
+
+//Returns true if there is a marker in interaction range
+bool FWaveformTransformationMarkerRenderer::SetMarkerInInteractionRange(const FVector2D& InLocalCursorPosition, const FGeometry& WidgetGeometry)
+{
+	InteractionRanges.Reset(MarkersArray->WaveCues.Num());
+	for (FSoundWaveCuePoint& CuePoint : MarkersArray->WaveCues)
+	{
+		if (IsPositionInInteractionRange(CuePoint.FramePosition, InLocalCursorPosition, WidgetGeometry))
+		{
+			MarkerInInteractionRange.SetMarkerInInteractionRange(&CuePoint);
+			return true;
+		}
+
+		if (CuePoint.FrameLength == 0)
+		{
+			continue;
+		}
+
+		const int32 LoopRegionEndHandlePosition = CuePoint.FramePosition + CuePoint.FrameLength;
+		
+		if (IsPositionInInteractionRange(LoopRegionEndHandlePosition, InLocalCursorPosition, WidgetGeometry))
+		{
+			MarkerInInteractionRange.SetMarkerInInteractionRange(&CuePoint, true);
+			return true;
+		}
+	}
+	MarkerInInteractionRange.SetMarkerInInteractionRange(nullptr);
+	return false;
+}
+
+FVector2D FWaveformTransformationMarkerRenderer::GetLocalCursorPosition(const FPointerEvent& MouseEvent, const FGeometry& EventGeometry) const
+{
+	const FVector2D ScreenSpacePosition = MouseEvent.GetScreenSpacePosition();
+	return  EventGeometry.AbsoluteToLocal(ScreenSpacePosition);
+}
+
+FWaveformTransformationMarkerRenderer::EMarkerInteractionType FWaveformTransformationMarkerRenderer::GetInteractionTypeFromCursorPosition(const FVector2D& InLocalCursorPosition, const FGeometry& WidgetGeometry)
+{
+	if (SetMarkerInInteractionRange(InLocalCursorPosition, WidgetGeometry))
+	{
+		return EMarkerInteractionType::ScrubbingMarkerHandle;
+	}
+
+	return EMarkerInteractionType::None;
+}
+
+void FWaveformTransformationMarkerRenderer::SetPropertyValueDependingOnInteractionType(const FPointerEvent& MouseEvent, const FGeometry& WidgetGeometry, const EPropertyValueSetFlags::Type Flags)
+{
+	const FVector2D LocalCursorPosition = GetLocalCursorPosition(MouseEvent, WidgetGeometry);
+	const float LocalCursorXClamped = FMath::Clamp(LocalCursorPosition.X, 0.f, WidgetGeometry.GetLocalSize().X);
+
+	if (MarkerInteractionType == EMarkerInteractionType::ScrubbingMarkerHandle)
+	{
+		if (MarkerInInteractionRange.GetMarker())
+		{
+			check(PixelsPerFrame > 0);
+			MarkerInInteractionRange.SetMarkerPosition(LocalCursorXClamped / PixelsPerFrame);
+		}
+	}
+}
+
+void FWaveformTransformationMarkerRenderer::FMarkerInInteractionRange::SetMarkerInInteractionRange(FSoundWaveCuePoint* InMarkerPtr, const bool InIsLoopRegionEndMarker)
+{
+	MarkerPtr = InMarkerPtr;
+	IsLoopRegionEndMarker = InIsLoopRegionEndMarker;
+}
+
+void FWaveformTransformationMarkerRenderer::FMarkerInInteractionRange::SetMarkerPosition(const float InPosition)
+{
+	check(MarkerPtr);
+
+	const int32 MinLoopSize = 10;
+
+	if (IsLoopRegionEndMarker)
+	{
+		//FMath::Max prevents the user from collapsing the loop
+		MarkerPtr->FrameLength = FMath::Max(InPosition, MarkerPtr->FramePosition + MinLoopSize) - MarkerPtr->FramePosition;
+	}
+	else
+	{
+		MarkerPtr->FramePosition = InPosition;
+	}
+}
diff --git a/Engine/Plugins/Editor/WaveformEditor/Source/WaveformTransformationsWidgets/Private/WaveformTransformationMarkersRenderer.h b/Engine/Plugins/Editor/WaveformEditor/Source/WaveformTransformationsWidgets/Private/WaveformTransformationMarkersRenderer.h
--- a/Engine/Plugins/Editor/WaveformEditor/Source/WaveformTransformationsWidgets/Private/WaveformTransformationMarkersRenderer.h
+++ b/Engine/Plugins/Editor/WaveformEditor/Source/WaveformTransformationsWidgets/Private/WaveformTransformationMarkersRenderer.h
@@ -13,6 +13,10 @@ public:
 
 	virtual int32 OnPaint(const FPaintArgs& Args, const FGeometry& AllottedGeometry, const FSlateRect& MyCullingRect, FSlateWindowElementList& OutDrawElements, int32 LayerId, const FWidgetStyle& InWidgetStyle, bool bParentEnabled) const override;
 	virtual void Tick(const FGeometry& AllottedGeometry, const double InCurrentTime, const float InDeltaTime) override;
+	virtual FReply OnMouseButtonDown(SWidget& OwnerWidget, const FGeometry& MyGeometry, const FPointerEvent& MouseEvent) override;
+	virtual FReply OnMouseButtonUp(SWidget& OwnerWidget, const FGeometry& MyGeometry, const FPointerEvent& MouseEvent) override;
+	virtual FReply OnMouseMove(SWidget& OwnerWidget, const FGeometry& MyGeometry, const FPointerEvent& MouseEvent) override;
+	virtual FCursorReply OnCursorQuery(const FGeometry& MyGeometry, const FPointerEvent& CursorEvent) const override;
 
 	virtual void SetPropertyHandles(const TArray<TSharedRef<IPropertyHandle>>& InPropertyHandles) override;
 
@@ -24,6 +28,33 @@ private:
 
 	void SetLoopColors(const TArray<FLinearColor>& InColors);
 
+	bool IsPositionInInteractionRange(const int32 InFramePosition, const FVector2D& InLocalCursorPosition, const FGeometry& WidgetGeometry) const;
+	bool SetMarkerInInteractionRange(const FVector2D& InLocalCursorPosition, const FGeometry& WidgetGeometry);
+	FVector2D GetLocalCursorPosition(const FPointerEvent& MouseEvent, const FGeometry& EventGeometry) const;
+
+	TArray<TRange<float>> InteractionRanges;
+
+	enum class EMarkerInteractionType : uint8
+	{
+		None = 0,
+		ScrubbingMarkerHandle
+	} MarkerInteractionType = EMarkerInteractionType::None;
+
+	class FMarkerInInteractionRange
+	{
+	public:
+		void SetMarkerInInteractionRange(FSoundWaveCuePoint* InMarkerPtr, const bool InIsLoopRegionEndMarker = false);
+		const FSoundWaveCuePoint* GetMarker() const { return MarkerPtr; }
+		void SetMarkerPosition(const float InPosition);
+		
+	private:
+		FSoundWaveCuePoint* MarkerPtr = nullptr;
+		bool IsLoopRegionEndMarker = false;
+	} MarkerInInteractionRange;
+
+	EMarkerInteractionType GetInteractionTypeFromCursorPosition(const FVector2D& InLocalCursorPosition, const FGeometry& WidgetGeometry);
+	void SetPropertyValueDependingOnInteractionType(const FPointerEvent& MouseEvent, const FGeometry& WidgetGeometry, const EPropertyValueSetFlags::Type Flags = EPropertyValueSetFlags::DefaultFlags);
+
 	TArray<FSoundWaveCuePoint> CuePoints;
 	TArray<FSoundWaveCuePoint> LoopPoints;
 
diff --git a/Engine/Plugins/Editor/WaveformEditor/Source/WaveformTransformationsWidgets/Private/WaveformTransformationTrimFadeRenderer.h b/Engine/Plugins/Editor/WaveformEditor/Source/WaveformTransformationsWidgets/Private/WaveformTransformationTrimFadeRenderer.h
--- a/Engine/Plugins/Editor/WaveformEditor/Source/WaveformTransformationsWidgets/Private/WaveformTransformationTrimFadeRenderer.h
+++ b/Engine/Plugins/Editor/WaveformEditor/Source/WaveformTransformationsWidgets/Private/WaveformTransformationTrimFadeRenderer.h
@@ -46,10 +46,6 @@ private:
 	TRange<float> FadeInInteractionXRange = TRange<float>::Inclusive(0.f, 0.f);
 	TRange<float> FadeOutInteractionXRange = TRange<float>::Inclusive(0.f, 0.f);
 
-	const float InteractionPixelXDelta = 10;
-	const float InteractionRatioYDelta = 0.07f;
-	const float MouseWheelStep = 0.25f;
-
 	double PixelsPerFrame = 0.0;
 
 	enum class ETrimFadeInteractionType : uint8
diff --git a/Engine/Plugins/Editor/WaveformEditor/Source/WaveformTransformationsWidgets/Public/WaveformTransformationRendererBase.h b/Engine/Plugins/Editor/WaveformEditor/Source/WaveformTransformationsWidgets/Public/WaveformTransformationRendererBase.h
--- a/Engine/Plugins/Editor/WaveformEditor/Source/WaveformTransformationsWidgets/Public/WaveformTransformationRendererBase.h
+++ b/Engine/Plugins/Editor/WaveformEditor/Source/WaveformTransformationsWidgets/Public/WaveformTransformationRendererBase.h
@@ -50,6 +50,10 @@ protected:
 
 	FWaveformTransformationRenderInfo TransformationWaveInfo;
 
+	static constexpr float InteractionPixelXDelta = 10;
+	static constexpr float InteractionRatioYDelta = 0.07f;
+	static constexpr float MouseWheelStep = 0.25f;
+
 private:
-	TArray<TSharedRef<IPropertyHandle>> CachedPropertyHandles;	
+	TArray<TSharedRef<IPropertyHandle>> CachedPropertyHandles;
 };
\ No newline at end of file
diff --git a/Engine/Source/Runtime/AudioExtensions/Public/IWaveformTransformation.h b/Engine/Source/Runtime/AudioExtensions/Public/IWaveformTransformation.h
--- a/Engine/Source/Runtime/AudioExtensions/Public/IWaveformTransformation.h
+++ b/Engine/Source/Runtime/AudioExtensions/Public/IWaveformTransformation.h
@@ -57,18 +57,18 @@ struct FSoundWaveCuePoint
 
 	// Unique identifier for the wave cue point
 	UPROPERTY(Category = Info, VisibleAnywhere, BlueprintReadOnly)
-	int32 CuePointID = 0;
+	int32 CuePointID = INDEX_NONE;
 
 	// The label for the cue point
-	UPROPERTY(Category = Info, VisibleAnywhere, BlueprintReadOnly)
+	UPROPERTY(Category = Info, EditAnywhere, BlueprintReadOnly)
 	FString Label;
 
 	// The frame position of the cue point
-	UPROPERTY(Category = Info, VisibleAnywhere, BlueprintReadOnly)
+	UPROPERTY(Category = Info, EditAnywhere, BlueprintReadOnly)
 	int32 FramePosition = 0;
 
 	// The frame length of the cue point (non-zero if it's a region)
-	UPROPERTY(Category = Info, VisibleAnywhere, BlueprintReadOnly)
+	UPROPERTY(Category = Info, EditAnywhere, BlueprintReadOnly)
 	int32 FrameLength = 0;
 
 	bool IsLoopRegion() const { return bIsLoopRegion; }
@@ -109,6 +109,7 @@ class UWaveformTransformationBase : public UObject
 public:
 	virtual Audio::FTransformationPtr CreateTransformation() const { return nullptr; }
 	virtual void UpdateConfiguration(FWaveTransformUObjectConfiguration& InOutConfiguration) {};
+	virtual void OverwriteTransformation() {};
 	
 	virtual bool IsEditorOnly() const override { return true; }
 };
