[Waveform Editor Bug Fix] - If all the cue points from an imported sound wave are removed, they will all automatically be readded when the array size is 0

diff --git a/Engine/Plugins/Editor/WaveformEditor/Source/WaveformTransformations/Private/WaveformTransformationMarkers.cpp b/Engine/Plugins/Editor/WaveformEditor/Source/WaveformTransformations/Private/WaveformTransformationMarkers.cpp
--- a/Engine/Plugins/Editor/WaveformEditor/Source/WaveformTransformations/Private/WaveformTransformationMarkers.cpp
+++ b/Engine/Plugins/Editor/WaveformEditor/Source/WaveformTransformations/Private/WaveformTransformationMarkers.cpp
@@ -18,17 +18,13 @@ Audio::FTransformationPtr UWaveformTransformationMarkers::CreateTransformation()
 void UWaveformTransformationMarkers::UpdateConfiguration(FWaveTransformUObjectConfiguration& InOutConfiguration)
 {
 	check(Markers != nullptr);
-
-	//Prevent USoundWave from overwriting the Transformation unintentionally
-	if(Markers->WaveCues.IsEmpty())
-	{
-		Markers->WaveCues = InOutConfiguration.WaveCues;
-	}
+	Markers->InitMarkersIfNotSet(InOutConfiguration.WaveCues);
 }
 
 void UWaveformTransformationMarkers::OverwriteTransformation()
 {
-	Markers->WaveCues.Empty();
+	check(Markers != nullptr);
+	Markers->Reset();
 }
 
 #if WITH_EDITOR
@@ -42,10 +38,10 @@ void UWaveCueArray::PostEditChangeProperty(FPropertyChangedEvent& PropertyChange
 	{
 		check(WaveCues.Num() > 0);
 
-		int32 LargestCuePointID = 0;
+		int32 LargestCuePointID = 1;
 		for (int32 Index = WaveCues.Num()-1; Index >= 0; Index--)
 		{
-			if (WaveCues[Index].CuePointID > LargestCuePointID)
+			if (WaveCues[Index].CuePointID >= LargestCuePointID)
 			{
 				LargestCuePointID = WaveCues[Index].CuePointID + 1;
 			}
@@ -60,4 +56,20 @@ void UWaveCueArray::PostEditChangeProperty(FPropertyChangedEvent& PropertyChange
 		}
 	}
 }
+
+void UWaveCueArray::InitMarkersIfNotSet(const TArray<FSoundWaveCuePoint>& InMarkers)
+{
+	//Prevent USoundWave from overwriting the Transformation unintentionally
+	if (!bIsInitialized)
+	{
+		WaveCues = InMarkers;
+		bIsInitialized = true;
+	}
+}
+
+void UWaveCueArray::Reset()
+{
+	WaveCues.Empty();
+	bIsInitialized = false;
+}
 #endif
\ No newline at end of file
diff --git a/Engine/Plugins/Editor/WaveformEditor/Source/WaveformTransformations/Public/WaveformTransformationMarkers.h b/Engine/Plugins/Editor/WaveformEditor/Source/WaveformTransformations/Public/WaveformTransformationMarkers.h
--- a/Engine/Plugins/Editor/WaveformEditor/Source/WaveformTransformations/Public/WaveformTransformationMarkers.h
+++ b/Engine/Plugins/Editor/WaveformEditor/Source/WaveformTransformations/Public/WaveformTransformationMarkers.h
@@ -13,9 +13,17 @@ class WAVEFORMTRANSFORMATIONS_API UWaveCueArray : public UObject
 	GENERATED_BODY()
 public:
 	void PostEditChangeProperty(FPropertyChangedEvent& PropertyChangedEvent);
+	
+	//If uninitialized, init wavecue array
+	void InitMarkersIfNotSet(const TArray<FSoundWaveCuePoint>& InMarkers);
+	//Uninitialize and empty wavecue array
+	void Reset();
 
 	UPROPERTY(EditAnywhere, Category = "Markers")
 	TArray<FSoundWaveCuePoint> WaveCues;
+
+private:
+	bool bIsInitialized = false;
 };
 
 class WAVEFORMTRANSFORMATIONS_API FWaveTransformationMarkers : public Audio::IWaveTransformation
