[Waveform Editor Bug fix for] - Deleting default transformations in editor preferences still remain in a Sound Wave when reopened after file is reloaded.
Extra context: when the default transformations are added, it doesn't mark the USoundWave as edited with the '*' even though it has been. But if you save it anyway, reload does not remove the transformations.

diff --git a/Engine/Plugins/Editor/WaveformEditor/Source/WaveformEditor/Private/WaveformEditor.cpp b/Engine/Plugins/Editor/WaveformEditor/Source/WaveformEditor/Private/WaveformEditor.cpp
--- a/Engine/Plugins/Editor/WaveformEditor/Source/WaveformEditor/Private/WaveformEditor.cpp
+++ b/Engine/Plugins/Editor/WaveformEditor/Source/WaveformEditor/Private/WaveformEditor.cpp
@@ -114,6 +114,10 @@ void FWaveformEditor::AddDefaultTransformations()
 					MaskedOuterFlags |= RF_ArchetypeObject;
 				}
 
+				//Adding default transformations does not mark the SoundWave to be saved
+				//Modify marks it as dirty to inform the user it needs to be saved if they want the default transformations to remain.
+				SoundWave->Modify();
+
 				UWaveformTransformationBase* TransformationToAdd = NewObject<UWaveformTransformationBase>(SoundWave, TransformationClass.Get(), NAME_None, MaskedOuterFlags);
 				SoundWave->Transformations.Add(TransformationToAdd);
 			}
@@ -124,6 +128,12 @@ void FWaveformEditor::AddDefaultTransformations()
 		}
 
 		TransformationChainConfig = SoundWave->UpdateTransformations();
+
+		//Update the content browser asset state to dirty
+		if (GEditor && !TransformationsSettings->LaunchTransformations.IsEmpty())
+		{
+			GEditor->BroadcastObjectReimported(SoundWave);
+		}
 	}
 }
 
