Waveform editor: allow modification frame length on markers without converting them to loop points.
Can now edit IsLoopRegion to create loop regions

diff --git a/Engine/Plugins/Editor/WaveformEditor/Source/WaveformTransformationsWidgets/Private/WaveformTransformationMarkersRenderer.cpp b/Engine/Plugins/Editor/WaveformEditor/Source/WaveformTransformationsWidgets/Private/WaveformTransformationMarkersRenderer.cpp
--- a/Engine/Plugins/Editor/WaveformEditor/Source/WaveformTransformationsWidgets/Private/WaveformTransformationMarkersRenderer.cpp
+++ b/Engine/Plugins/Editor/WaveformEditor/Source/WaveformTransformationsWidgets/Private/WaveformTransformationMarkersRenderer.cpp
@@ -217,7 +217,7 @@ void FWaveformTransformationMarkerRenderer::Tick(const FGeometry& AllottedGeomet
 	LoopPoints.Empty();
 	for (TPair<uint32, FSoundWaveCuePoint>& Pair : CuePointMap)
 	{
-		if (Pair.Value.FrameLength > 0)
+		if (Pair.Value.IsLoopRegion())
 		{
 			LoopPoints.Add(MoveTemp(Pair.Value));
 		}
diff --git a/Engine/Source/Runtime/AudioExtensions/Public/IWaveformTransformation.h b/Engine/Source/Runtime/AudioExtensions/Public/IWaveformTransformation.h
--- a/Engine/Source/Runtime/AudioExtensions/Public/IWaveformTransformation.h
+++ b/Engine/Source/Runtime/AudioExtensions/Public/IWaveformTransformation.h
@@ -86,7 +86,7 @@ struct FSoundWaveCuePoint
 private:
 	// intentionally kept private.
 	// only USoundFactory should modify this value on import
-	UPROPERTY(Category = Info, VisibleAnywhere)
+	UPROPERTY(Category = Info, EditAnywhere)
 	bool bIsLoopRegion = false;
 };
 
