Waveform Editor: Fade in and Fade out

diff --git a/Engine/Plugins/Editor/WaveformEditor/Source/WaveformEditor/Private/WaveformEditor.cpp b/Engine/Plugins/Editor/WaveformEditor/Source/WaveformEditor/Private/WaveformEditor.cpp
--- a/Engine/Plugins/Editor/WaveformEditor/Source/WaveformEditor/Private/WaveformEditor.cpp
+++ b/Engine/Plugins/Editor/WaveformEditor/Source/WaveformEditor/Private/WaveformEditor.cpp
@@ -12,6 +12,7 @@
 #include "SparseSampledSequenceTransportCoordinator.h"
 #include "STransformedWaveformViewPanel.h"
 #include "Styling/AppStyle.h"
+#include "Styling/SlateStyleRegistry.h"
 #include "SWaveformTransformationsOverlay.h"
 #include "ToolMenus.h"
 #include "TransformedWaveformView.h"
@@ -137,6 +138,141 @@ void FWaveformEditor::AddDefaultTransformations()
 	}
 }
 
+void FWaveformEditor::NotifyPostTransformationChange(const EPropertyChangeType::Type& PropertyChangeType)
+{
+	check(SoundWave);
+
+	FProperty* TransformationsProperty = FindFProperty<FProperty>(USoundWave::StaticClass(), GET_MEMBER_NAME_CHECKED(USoundWave, Transformations));
+
+	FEditPropertyChain PropertyChain;
+	PropertyChain.AddHead(TransformationsProperty);
+	PropertyChain.SetActivePropertyNode(TransformationsProperty);
+
+	FPropertyChangedEvent PropertyChangedEvent(TransformationsProperty, PropertyChangeType);
+
+	NotifyPostChange(PropertyChangedEvent, &PropertyChain);
+
+	//Update the content browser asset state to dirty
+	if (GEditor)
+	{
+		GEditor->BroadcastObjectReimported(SoundWave);
+	}
+
+	TransformationChainConfig = SoundWave->UpdateTransformations();
+}
+
+TObjectPtr<UWaveformTransformationTrimFade> FWaveformEditor::GetOrAddTrimFadeTransformation()
+{
+	check(SoundWave);
+
+	TObjectPtr<UWaveformTransformationTrimFade> TrimFadeTransformation;
+
+	for (TObjectPtr<UWaveformTransformationBase> Transformation : SoundWave->Transformations)
+	{
+		TrimFadeTransformation = Cast<UWaveformTransformationTrimFade>(Transformation);
+		if (TrimFadeTransformation)
+		{
+			break;
+		}
+	}
+
+	if (TrimFadeTransformation == nullptr)
+	{
+		EObjectFlags MaskedOuterFlags = SoundWave ? SoundWave->GetMaskedFlags(RF_PropagateToSubObjects) : RF_NoFlags;
+
+		if (SoundWave->HasAnyFlags(RF_ClassDefaultObject | RF_ArchetypeObject))
+		{
+			MaskedOuterFlags |= RF_ArchetypeObject;
+		}
+
+		TrimFadeTransformation = NewObject<UWaveformTransformationTrimFade>(SoundWave, UWaveformTransformationTrimFade::StaticClass(), NAME_None, MaskedOuterFlags);
+
+		SoundWave->Modify();
+		SoundWave->Transformations.Add(TrimFadeTransformation);
+
+		NotifyPostTransformationChange();
+	}
+	return TrimFadeTransformation;
+}
+
+void FWaveformEditor::ToggleFadeIn()
+{
+	check(SoundWave);
+
+	TObjectPtr<UWaveformTransformationTrimFade> TrimFadeTransformation = GetOrAddTrimFadeTransformation();
+	check(TrimFadeTransformation);
+
+	SoundWave->Modify();
+	
+	//if fade duration is 0, assume there is no fade present and toggle on the fade with the curve data selected in the menu
+	if (TrimFadeTransformation->StartFadeTime == 0.0f)
+	{
+		check(DefaultFadeInAmount >= 0);
+		check(FadeInCurve >= -0.1f && FadeInCurve <= 10);
+
+		if (CachedFadeInAmount <= 0)
+		{
+			CachedFadeInAmount = DefaultFadeInAmount;
+		}
+
+		TrimFadeTransformation->StartFadeTime = CachedFadeInAmount;
+		TrimFadeTransformation->StartFadeCurve = FadeInCurve;
+	}
+	else //if fade duration is not 0, assume there is a fade applied and toggle fade "off" by setting fade time to 0
+	{
+		CachedFadeInAmount = TrimFadeTransformation->StartFadeTime;
+
+		//No need to set FadeCurve when toggling off. This will preserve user curve data if they later manually change the StartFadeTime.
+		TrimFadeTransformation->StartFadeTime = 0.0f;
+	}
+
+	NotifyPostTransformationChange(EPropertyChangeType::ValueSet);
+}
+
+bool FWaveformEditor::CanFadeIn()
+{
+	return true;
+}
+
+void FWaveformEditor::ToggleFadeOut()
+{
+	check(SoundWave);
+
+	TObjectPtr<UWaveformTransformationTrimFade> TrimFadeTransformation = GetOrAddTrimFadeTransformation();
+	check(TrimFadeTransformation);
+
+	SoundWave->Modify();
+
+	//if fade duration is 0, assume there is no fade present and toggle on the fade with the curve data selected in the menu
+	if (TrimFadeTransformation->EndFadeTime == 0.0f)
+	{
+		check(DefaultFadeOutAmount >= 0);
+		check(FadeOutCurve >= -0.1f && FadeOutCurve <= 10);
+
+		if (CachedFadeOutAmount <= 0)
+		{
+			CachedFadeOutAmount = DefaultFadeOutAmount;
+		}
+		
+		TrimFadeTransformation->EndFadeTime = CachedFadeOutAmount;
+		TrimFadeTransformation->EndFadeCurve = FadeOutCurve;
+	}
+	else //if fade duration is not 0, assume there is a fade applied and toggle fade "off" by setting fade time to 0
+	{
+		CachedFadeOutAmount = TrimFadeTransformation->EndFadeTime;
+
+		//No need to set FadeCurve when toggling off. This will preserve user curve data if they later manually change the EndFadeTime.
+		TrimFadeTransformation->EndFadeTime = 0.0f;
+	}
+
+	NotifyPostTransformationChange(EPropertyChangeType::ValueSet);
+}
+
+bool FWaveformEditor::CanFadeOut()
+{
+	return true;
+}
+
 bool FWaveformEditor::InitializeAudioComponent()
 {
 	if (!ensure(SoundWave))
@@ -346,14 +482,14 @@ bool FWaveformEditor::RegisterToolbar()
 
 		ExportSection.AddDynamicEntry("ExportButton", FNewToolMenuSectionDelegate::CreateLambda([this, Commands](FToolMenuSection& InSection) 
 		{
-			UWaveformEditorToolMenuContext* Context = InSection.FindContext<UWaveformEditorToolMenuContext>();
+				const UWaveformEditorToolMenuContext* Context = InSection.FindContext<UWaveformEditorToolMenuContext>();
 
-			if (Context && Context->WaveformEditor.IsValid())
+			if (TSharedPtr<FWaveformEditor> LockedObserver = Context->WaveformEditor.Pin())
 			{
 				FToolMenuEntry ExportEntry = FToolMenuEntry::InitToolBarButton(
 					Commands.ExportWaveform,
 					FText(),
-					TAttribute< FText >::CreateRaw(Context->WaveformEditor.Pin().Get(), &FWaveformEditor::GetExportButtonToolTip),
+					TAttribute< FText >::CreateRaw(LockedObserver.Get(), &FWaveformEditor::GetExportButtonToolTip),
 					FSlateIcon(FAppStyle::GetAppStyleSetName(), "LevelEditor.ExportAll")
 				);
 
@@ -363,14 +499,14 @@ bool FWaveformEditor::RegisterToolbar()
 
 		ExportSection.AddDynamicEntry("ExportOptions", FNewToolMenuSectionDelegate::CreateLambda([this](FToolMenuSection& InSection) 
 		{
-				UWaveformEditorToolMenuContext* Context = InSection.FindContext<UWaveformEditorToolMenuContext>();
+				const UWaveformEditorToolMenuContext* Context = InSection.FindContext<UWaveformEditorToolMenuContext>();
 				
-				if (Context && Context->WaveformEditor.IsValid())
+				if (TSharedPtr<FWaveformEditor> LockedObserver = Context->WaveformEditor.Pin())
 				{
 					FToolMenuEntry ExportOptionsEntry = FToolMenuEntry::InitComboButton(
 						"ExportsOptionsCombo",
 						FToolUIActionChoice(FUIAction()),
-						FNewToolMenuChoice(FOnGetContent::CreateSP(Context->WaveformEditor.Pin().Get(), &FWaveformEditor::GenerateExportOptionsMenu)),
+						FNewToolMenuChoice(FOnGetContent::CreateSP(LockedObserver.Get(), &FWaveformEditor::GenerateExportOptionsMenu)),
 						LOCTEXT("ExportsOptions_Label", "Export Options"),
 						LOCTEXT("ExportsOptions_ToolTip", "Export options for this waveform"),
 						FSlateIcon(FAppStyle::GetAppStyleSetName(), "LevelEditor.ExportAll"),
@@ -383,14 +519,14 @@ bool FWaveformEditor::RegisterToolbar()
 
 		ExportSection.AddDynamicEntry("ImportButton", FNewToolMenuSectionDelegate::CreateLambda([this, Commands](FToolMenuSection& InSection)
 		{
-			UWaveformEditorToolMenuContext* Context = InSection.FindContext<UWaveformEditorToolMenuContext>();
+			const UWaveformEditorToolMenuContext* Context = InSection.FindContext<UWaveformEditorToolMenuContext>();
 
-			if (Context && Context->WaveformEditor.IsValid())
+			if (TSharedPtr<FWaveformEditor> LockedObserver = Context->WaveformEditor.Pin())
 			{
 				FToolMenuEntry ReimportEntry = FToolMenuEntry::InitToolBarButton(
 					Commands.ReimportAsset,
 					FText(),
-					TAttribute< FText >::CreateRaw(Context->WaveformEditor.Pin().Get(), &FWaveformEditor::GetReimportButtonToolTip),
+					TAttribute< FText >::CreateRaw(LockedObserver.Get(), &FWaveformEditor::GetReimportButtonToolTip),
 					FSlateIcon(FAppStyle::GetAppStyleSetName(), "Icons.Import"));
 
 				InSection.AddEntry(ReimportEntry);
@@ -399,14 +535,14 @@ bool FWaveformEditor::RegisterToolbar()
 
 		ExportSection.AddDynamicEntry("ReimportOptions", FNewToolMenuSectionDelegate::CreateLambda([this](FToolMenuSection& InSection) 
 		{
-			UWaveformEditorToolMenuContext* Context = InSection.FindContext<UWaveformEditorToolMenuContext>();
+			const UWaveformEditorToolMenuContext* Context = InSection.FindContext<UWaveformEditorToolMenuContext>();
 
-			if (Context && Context->WaveformEditor.IsValid())
+			if (TSharedPtr<FWaveformEditor> LockedObserver = Context->WaveformEditor.Pin())
 			{
 				FToolMenuEntry ReimportOptionsEntry = FToolMenuEntry::InitComboButton(
 					"ReimportOptionsCombo",
 					FToolUIActionChoice(FUIAction()),
-					FNewToolMenuChoice(FOnGetContent::CreateSP(Context->WaveformEditor.Pin().Get(), &FWaveformEditor::GenerateImportOptionsMenu)),
+					FNewToolMenuChoice(FOnGetContent::CreateSP(LockedObserver.Get(), &FWaveformEditor::GenerateImportOptionsMenu)),
 					LOCTEXT("ReimportOptions_Label", "Reimport Options"),
 					LOCTEXT("ReimportOptions_ToolTip", "Reimport options for this USoundWave"),
 					FSlateIcon(FAppStyle::GetAppStyleSetName(), "Icons.Import"),
@@ -417,11 +553,130 @@ bool FWaveformEditor::RegisterToolbar()
 			}
 
 		}));
+
+		FToolMenuInsert InsertAfterExportSection("Export Controls", EToolMenuInsertType::After);
+		FToolMenuSection& TransformationSection = ToolBar->AddSection("Transformation Controls", TAttribute<FText>(), InsertAfterExportSection);
+
+		const FSlateBrush* IconBrush = FAppStyle::Get().GetBrush("Sequencer.Tracks.Fade");
+		FSlateBrush* FlippedIconBrush = new FSlateBrush(*IconBrush);
+		FlippedIconBrush->Mirroring = ESlateBrushMirrorType::Horizontal;
+
+		// Register the flipped brush in a custom style set
+		static TSharedRef<FSlateStyleSet> FlippedStyleSet = MakeShared<FSlateStyleSet>("FlippedStyleSet");
+		FlippedStyleSet->SetContentRoot(FAppStyle::Get().GetContentRootDir());
+		FlippedStyleSet->Set("Sequencer.Tracks.Fade.Flipped", FlippedIconBrush);
+		FSlateStyleRegistry::RegisterSlateStyle(*FlippedStyleSet);
+
+		TransformationSection.AddDynamicEntry("ToggleFadeInButton", FNewToolMenuSectionDelegate::CreateLambda([this, Commands](FToolMenuSection& InSection)
+			{
+				const UWaveformEditorToolMenuContext* Context = InSection.FindContext<UWaveformEditorToolMenuContext>();
+
+				if (Context && Context->WaveformEditor.IsValid())
+				{
+					FToolMenuEntry ToggleFadeIn = FToolMenuEntry::InitToolBarButton(
+						Commands.ToggleFadeIn,
+						FText(),
+						LOCTEXT("WaveformEditorFadeInButtonTooltip", "Toggle Fade In Transformation onto the soundwave"),
+						FSlateIcon("FlippedStyleSet", "Sequencer.Tracks.Fade.Flipped")
+					);
+					InSection.AddEntry(ToggleFadeIn);
+				}
+			}));
+
+		TransformationSection.AddDynamicEntry("FadeInOptions", FNewToolMenuSectionDelegate::CreateLambda([this](FToolMenuSection& InSection)
+			{
+				const UWaveformEditorToolMenuContext* Context = InSection.FindContext<UWaveformEditorToolMenuContext>();
+
+				if (TSharedPtr<FWaveformEditor> LockedObserver = Context->WaveformEditor.Pin())
+				{
+					FToolMenuEntry ExportOptionsEntry = FToolMenuEntry::InitComboButton(
+						"FadeInOptionsCombo",
+						FToolUIActionChoice(FUIAction()),
+						FNewToolMenuChoice(FOnGetContent::CreateSP(LockedObserver.Get(), &FWaveformEditor::GenerateFadeInOptionsMenu)),
+						LOCTEXT("FadeInOptions_Label", "Fade In Options"),
+						LOCTEXT("FadeInOptions_ToolTip", "Fade In options for this waveform"),
+						FSlateIcon(FAppStyle::GetAppStyleSetName(), "LevelEditor.ExportAll"),
+						true
+					);
+
+					InSection.AddEntry(ExportOptionsEntry);
+				}
+			}));
+
+		TransformationSection.AddDynamicEntry("ToggleFadeOutButton", FNewToolMenuSectionDelegate::CreateLambda([this, Commands](FToolMenuSection& InSection)
+			{
+				const UWaveformEditorToolMenuContext* Context = InSection.FindContext<UWaveformEditorToolMenuContext>();
+
+				if (Context && Context->WaveformEditor.IsValid())
+				{
+					FToolMenuEntry ToggleFadeOut = FToolMenuEntry::InitToolBarButton(
+						Commands.ToggleFadeOut,
+						FText(),
+						LOCTEXT("WaveformEditorFadeOutButtonTooltip", "Toggle Fade Out Transformation onto the soundwave"),
+						FSlateIcon(FAppStyle::GetAppStyleSetName(), "Sequencer.Tracks.Fade")
+					);
+					InSection.AddEntry(ToggleFadeOut);
+				}
+			}));
+
+		TransformationSection.AddDynamicEntry("FadeOutOptions", FNewToolMenuSectionDelegate::CreateLambda([this](FToolMenuSection& InSection)
+			{
+				const UWaveformEditorToolMenuContext* Context = InSection.FindContext<UWaveformEditorToolMenuContext>();
+
+				if (TSharedPtr<FWaveformEditor> LockedObserver = Context->WaveformEditor.Pin())
+				{
+					FToolMenuEntry ExportOptionsEntry = FToolMenuEntry::InitComboButton(
+						"FadeOutOptionsCombo",
+						FToolUIActionChoice(FUIAction()),
+						FNewToolMenuChoice(FOnGetContent::CreateSP(LockedObserver.Get(), &FWaveformEditor::GenerateFadeOutOptionsMenu)),
+						LOCTEXT("FadeOutOptions_Label", "Fade Out Options"),
+						LOCTEXT("FadeOutOptions_ToolTip", "Fade Out options for this waveform"),
+						FSlateIcon(FAppStyle::GetAppStyleSetName(), "LevelEditor.ExportAll"),
+						true
+					);
+
+					InSection.AddEntry(ExportOptionsEntry);
+				}
+			}));
 	}
 
 	return true;
 }
 
+TSharedRef<SWidget> FWaveformEditor::GenerateFadeInOptionsMenu()
+{
+	const bool bShouldCloseWindowAfterMenuSelection = true;
+	FMenuBuilder MenuBuilder(bShouldCloseWindowAfterMenuSelection, GetToolkitCommands());
+
+	MenuBuilder.BeginSection(NAME_None, LOCTEXT("FadeInOptionsSection_Label", "Fade In Options"));
+	{
+		MenuBuilder.AddMenuEntry(FWaveformEditorCommands::Get().FadeInLinear);
+		MenuBuilder.AddMenuEntry(FWaveformEditorCommands::Get().FadeInExponential);
+		MenuBuilder.AddMenuEntry(FWaveformEditorCommands::Get().FadeInLogarithmic);
+		MenuBuilder.AddMenuEntry(FWaveformEditorCommands::Get().FadeInSigmoid);
+	}
+	MenuBuilder.EndSection();
+
+	return MenuBuilder.MakeWidget();
+}
+
+TSharedRef<SWidget> FWaveformEditor::GenerateFadeOutOptionsMenu()
+{
+	const bool bShouldCloseWindowAfterMenuSelection = true;
+	FMenuBuilder MenuBuilder(bShouldCloseWindowAfterMenuSelection, GetToolkitCommands());
+
+	MenuBuilder.BeginSection(NAME_None, LOCTEXT("FadeOutOptionsSection_Label", "Fade Out Options"));
+	{
+		MenuBuilder.AddMenuEntry(FWaveformEditorCommands::Get().FadeOutLinear);
+		MenuBuilder.AddMenuEntry(FWaveformEditorCommands::Get().FadeOutExponential);
+		MenuBuilder.AddMenuEntry(FWaveformEditorCommands::Get().FadeOutLogarithmic);
+		MenuBuilder.AddMenuEntry(FWaveformEditorCommands::Get().FadeOutSigmoid);
+	}
+	MenuBuilder.EndSection();
+
+	return MenuBuilder.MakeWidget();
+}
+
 TSharedRef<SWidget> FWaveformEditor::GenerateExportOptionsMenu()
 {
 	const bool bShouldCloseWindowAfterMenuSelection = true;
@@ -536,6 +791,64 @@ bool FWaveformEditor::BindCommands()
 		FCanExecuteAction::CreateLambda([this] { return CanExecuteReimport(); }),
 		FIsActionChecked::CreateLambda([this] {return ReimportMode == EWaveEditorReimportMode::SelectFile; }));
 
+	ToolkitCommands->MapAction(
+		Commands.ToggleFadeIn,
+		FExecuteAction::CreateSP(this, &FWaveformEditor::ToggleFadeIn),
+		FCanExecuteAction::CreateSP(this, &FWaveformEditor::CanFadeIn));
+
+	ToolkitCommands->MapAction(
+		Commands.FadeInLinear,
+		FExecuteAction::CreateLambda([this] { FadeInCurve = UWaveformTransformationTrimFade::FadeModeToCurveValueMap[EWaveEditorFadeMode::Linear]; }),
+		FCanExecuteAction::CreateLambda([this] {return true; }),
+		FIsActionChecked::CreateLambda([this] {return FadeInCurve == UWaveformTransformationTrimFade::FadeModeToCurveValueMap[EWaveEditorFadeMode::Linear]; }));
+
+	ToolkitCommands->MapAction(
+		Commands.FadeInExponential,
+		FExecuteAction::CreateLambda([this] { FadeInCurve = UWaveformTransformationTrimFade::FadeModeToCurveValueMap[EWaveEditorFadeMode::Exponetial]; }),
+		FCanExecuteAction::CreateLambda([this] {return true; }),
+		FIsActionChecked::CreateLambda([this] {return FadeInCurve == UWaveformTransformationTrimFade::FadeModeToCurveValueMap[EWaveEditorFadeMode::Exponetial]; }));
+
+	ToolkitCommands->MapAction(
+		Commands.FadeInLogarithmic,
+		FExecuteAction::CreateLambda([this] { FadeInCurve = UWaveformTransformationTrimFade::FadeModeToCurveValueMap[EWaveEditorFadeMode::Logarithmic]; }),
+		FCanExecuteAction::CreateLambda([this] {return true; }),
+		FIsActionChecked::CreateLambda([this] {return FadeInCurve == UWaveformTransformationTrimFade::FadeModeToCurveValueMap[EWaveEditorFadeMode::Logarithmic]; }));
+
+	ToolkitCommands->MapAction(
+		Commands.FadeInSigmoid,
+		FExecuteAction::CreateLambda([this] { FadeInCurve = UWaveformTransformationTrimFade::FadeModeToCurveValueMap[EWaveEditorFadeMode::Sigmoid]; }),
+		FCanExecuteAction::CreateLambda([this] {return true; }),
+		FIsActionChecked::CreateLambda([this] {return FadeInCurve == UWaveformTransformationTrimFade::FadeModeToCurveValueMap[EWaveEditorFadeMode::Sigmoid]; }));
+
+	ToolkitCommands->MapAction(
+		Commands.ToggleFadeOut,
+		FExecuteAction::CreateSP(this, &FWaveformEditor::ToggleFadeOut),
+		FCanExecuteAction::CreateSP(this, &FWaveformEditor::CanFadeOut));
+
+	ToolkitCommands->MapAction(
+		Commands.FadeOutLinear,
+		FExecuteAction::CreateLambda([this] { FadeOutCurve = UWaveformTransformationTrimFade::FadeModeToCurveValueMap[EWaveEditorFadeMode::Linear]; }),
+		FCanExecuteAction::CreateLambda([this] {return true; }),
+		FIsActionChecked::CreateLambda([this] {return FadeOutCurve == UWaveformTransformationTrimFade::FadeModeToCurveValueMap[EWaveEditorFadeMode::Linear]; }));
+
+	ToolkitCommands->MapAction(
+		Commands.FadeOutExponential,
+		FExecuteAction::CreateLambda([this] { FadeOutCurve = UWaveformTransformationTrimFade::FadeModeToCurveValueMap[EWaveEditorFadeMode::Exponetial]; }),
+		FCanExecuteAction::CreateLambda([this] {return true; }),
+		FIsActionChecked::CreateLambda([this] {return FadeOutCurve == UWaveformTransformationTrimFade::FadeModeToCurveValueMap[EWaveEditorFadeMode::Exponetial]; }));
+
+	ToolkitCommands->MapAction(
+		Commands.FadeOutLogarithmic,
+		FExecuteAction::CreateLambda([this] { FadeOutCurve = UWaveformTransformationTrimFade::FadeModeToCurveValueMap[EWaveEditorFadeMode::Logarithmic]; }),
+		FCanExecuteAction::CreateLambda([this] {return true; }),
+		FIsActionChecked::CreateLambda([this] {return FadeOutCurve == UWaveformTransformationTrimFade::FadeModeToCurveValueMap[EWaveEditorFadeMode::Logarithmic]; }));
+
+	ToolkitCommands->MapAction(
+		Commands.FadeOutSigmoid,
+		FExecuteAction::CreateLambda([this] { FadeOutCurve = UWaveformTransformationTrimFade::FadeModeToCurveValueMap[EWaveEditorFadeMode::Sigmoid]; }),
+		FCanExecuteAction::CreateLambda([this] {return true; }),
+		FIsActionChecked::CreateLambda([this] {return FadeOutCurve == UWaveformTransformationTrimFade::FadeModeToCurveValueMap[EWaveEditorFadeMode::Sigmoid]; }));
+
 	return true;
 }
 
diff --git a/Engine/Plugins/Editor/WaveformEditor/Source/WaveformEditor/Private/WaveformEditorCommands.cpp b/Engine/Plugins/Editor/WaveformEditor/Source/WaveformEditor/Private/WaveformEditorCommands.cpp
--- a/Engine/Plugins/Editor/WaveformEditor/Source/WaveformEditor/Private/WaveformEditorCommands.cpp
+++ b/Engine/Plugins/Editor/WaveformEditor/Source/WaveformEditor/Private/WaveformEditorCommands.cpp
@@ -31,6 +31,19 @@ void FWaveformEditorCommands::RegisterCommands()
 	UI_COMMAND(ReimportModeSameFile, "Reimport from same file", "Reimports asset from the same file", EUserInterfaceActionType::Check, FInputChord());
 	UI_COMMAND(ReimportModeSameFileOverwriteTransformations, "Reimport from same file and overwrite transformations", "Reimports asset from the same file and overwrites the transformations", EUserInterfaceActionType::Check, FInputChord());
 	UI_COMMAND(ReimportModeNewFile, "Reimport from new file", "Reimports asset from a selected file", EUserInterfaceActionType::Check, FInputChord());
+
+	UI_COMMAND(ToggleFadeIn, "Toggle Fade In", "Toggle adding Fade In Transformation onto the sound wave", EUserInterfaceActionType::Button, FInputChord());
+	UI_COMMAND(FadeInLinear, "Linear Fade In", "Fade In attack curve: 1.0 = linear growth, < 1.0 logorithmic growth, > 1.0 exponential growth", EUserInterfaceActionType::RadioButton, FInputChord());
+	UI_COMMAND(FadeInExponential, "Exponential Fade In", "Fade In attack curve: 1.0 = linear growth, < 1.0 logarithmic growth, > 1.0 exponential growth", EUserInterfaceActionType::RadioButton, FInputChord());
+	UI_COMMAND(FadeInLogarithmic, "Logarithmic Fade In", "Fade In attack curve: 1.0 = linear growth, < 1.0 logarithmic growth, > 1.0 exponential growth", EUserInterfaceActionType::RadioButton, FInputChord());
+	UI_COMMAND(FadeInSigmoid, "Sigmoid Fade In", "Fade In S-Shaped curve", EUserInterfaceActionType::RadioButton, FInputChord());
+	
+	UI_COMMAND(ToggleFadeOut, "Toggle Fade Out", "Toggle adding Fade Out Transformation onto the sound wave", EUserInterfaceActionType::Button, FInputChord());
+	UI_COMMAND(FadeOutLinear, "Linear Fade Out", "Fade Out attack curve: 1.0 = linear growth, < 1.0 logorithmic growth, > 1.0 exponential growth", EUserInterfaceActionType::RadioButton, FInputChord());
+	UI_COMMAND(FadeOutExponential, "Exponential Fade Out", "Fade Out attack curve: 1.0 = linear growth, < 1.0 logarithmic growth, > 1.0 exponential growth", EUserInterfaceActionType::RadioButton, FInputChord());
+	UI_COMMAND(FadeOutLogarithmic, "Logarithmic Fade Out", "Fade Out attack curve: 1.0 = linear growth, < 1.0 logarithmic growth, > 1.0 exponential growth", EUserInterfaceActionType::RadioButton, FInputChord());
+	UI_COMMAND(FadeOutSigmoid, "Sigmoid Fade Out", "Fade Out S-Shaped curve", EUserInterfaceActionType::RadioButton, FInputChord());
+
 }
 	
 const FWaveformEditorCommands& FWaveformEditorCommands::Get()
diff --git a/Engine/Plugins/Editor/WaveformEditor/Source/WaveformEditor/Public/WaveformEditor.h b/Engine/Plugins/Editor/WaveformEditor/Source/WaveformEditor/Public/WaveformEditor.h
--- a/Engine/Plugins/Editor/WaveformEditor/Source/WaveformEditor/Public/WaveformEditor.h
+++ b/Engine/Plugins/Editor/WaveformEditor/Source/WaveformEditor/Public/WaveformEditor.h
@@ -8,6 +8,7 @@
 #include "Toolkits/AssetEditorToolkit.h"
 #include "WaveformEditorTransportController.h"
 #include "WaveformEditorZoomController.h"
+#include "WaveformTransformationTrimFade.h"
 #include "TransformedWaveformView.h"
 
 class IToolkitHost;
@@ -64,6 +65,8 @@ private:
 	/**	Toolbar Setup */
 	bool RegisterToolbar();
 	bool BindCommands();
+	TSharedRef<SWidget> GenerateFadeInOptionsMenu();
+	TSharedRef<SWidget> GenerateFadeOutOptionsMenu();
 	TSharedRef<SWidget> GenerateExportOptionsMenu();
 	TSharedRef<SWidget> GenerateImportOptionsMenu();
 	bool CanExecuteReimport() const;
@@ -98,6 +101,16 @@ private:
 	const UWaveformEditorTransformationsSettings* GetWaveformEditorTransformationsSettings() const;
 	void AddDefaultTransformations();
 
+	/** Transformation Functions*/
+	void NotifyPostTransformationChange(const EPropertyChangeType::Type& PropertyChangeType = EPropertyChangeType::ArrayAdd);
+	TObjectPtr<UWaveformTransformationTrimFade> GetOrAddTrimFadeTransformation();
+
+	void ToggleFadeIn();
+	bool CanFadeIn();
+
+	void ToggleFadeOut();
+	bool CanFadeOut();
+
 	FTransformedWaveformView  WaveformView;
 
 	/** Exports the edited waveform to a new asset */
@@ -137,6 +150,14 @@ private:
 	float StartTimeBeforeTransformInteraction = 0.f;
 	FWaveTransformUObjectConfiguration TransformationChainConfig;
 
+	static constexpr float DefaultFadeInAmount = 0.5f;
+	float CachedFadeInAmount = 0.f;
+	float FadeInCurve = UWaveformTransformationTrimFade::FadeModeToCurveValueMap[EWaveEditorFadeMode::Exponetial];
+
+	static constexpr float DefaultFadeOutAmount = 0.5f;
+	float CachedFadeOutAmount = 0.f;
+	float FadeOutCurve = UWaveformTransformationTrimFade::FadeModeToCurveValueMap[EWaveEditorFadeMode::Exponetial];
+
 	enum class EWaveEditorReimportMode : uint8
 	{
 		SameFile = 0, 
diff --git a/Engine/Plugins/Editor/WaveformEditor/Source/WaveformEditor/Public/WaveformEditorCommands.h b/Engine/Plugins/Editor/WaveformEditor/Source/WaveformEditor/Public/WaveformEditorCommands.h
--- a/Engine/Plugins/Editor/WaveformEditor/Source/WaveformEditor/Public/WaveformEditorCommands.h
+++ b/Engine/Plugins/Editor/WaveformEditor/Source/WaveformEditor/Public/WaveformEditorCommands.h
@@ -32,4 +32,15 @@ public:
 	TSharedPtr<FUICommandInfo> ReimportModeSameFileOverwriteTransformations;
 	TSharedPtr<FUICommandInfo> ReimportModeNewFile;
 
+	TSharedPtr<FUICommandInfo> ToggleFadeIn;
+	TSharedPtr<FUICommandInfo> FadeInLinear;
+	TSharedPtr<FUICommandInfo> FadeInExponential;
+	TSharedPtr<FUICommandInfo> FadeInLogarithmic;
+	TSharedPtr<FUICommandInfo> FadeInSigmoid;
+
+	TSharedPtr<FUICommandInfo> ToggleFadeOut;
+	TSharedPtr<FUICommandInfo> FadeOutLinear;
+	TSharedPtr<FUICommandInfo> FadeOutExponential;
+	TSharedPtr<FUICommandInfo> FadeOutLogarithmic;
+	TSharedPtr<FUICommandInfo> FadeOutSigmoid;
 };
diff --git a/Engine/Plugins/Editor/WaveformEditor/Source/WaveformEditor/WaveformEditor.build.cs b/Engine/Plugins/Editor/WaveformEditor/Source/WaveformEditor/WaveformEditor.build.cs
--- a/Engine/Plugins/Editor/WaveformEditor/Source/WaveformEditor/WaveformEditor.build.cs
+++ b/Engine/Plugins/Editor/WaveformEditor/Source/WaveformEditor/WaveformEditor.build.cs
@@ -29,6 +29,7 @@ namespace UnrealBuildTool.Rules
 					"ToolMenus",
 					"UnrealEd", 
 					"WaveformEditorWidgets",
+					"WaveformTransformations",
 					"WaveformTransformationsWidgets"
 				}
 			);
diff --git a/Engine/Plugins/Editor/WaveformEditor/Source/WaveformTransformations/Private/WaveformTransformationTrimFade.cpp b/Engine/Plugins/Editor/WaveformEditor/Source/WaveformTransformations/Private/WaveformTransformationTrimFade.cpp
--- a/Engine/Plugins/Editor/WaveformEditor/Source/WaveformTransformations/Private/WaveformTransformationTrimFade.cpp
+++ b/Engine/Plugins/Editor/WaveformEditor/Source/WaveformTransformations/Private/WaveformTransformationTrimFade.cpp
@@ -3,30 +3,38 @@
 #include "WaveformTransformationTrimFade.h"
 #include "Sound/SoundWave.h"
 
+const TMap<EWaveEditorFadeMode, float> UWaveformTransformationTrimFade::FadeModeToCurveValueMap =
+{
+	{EWaveEditorFadeMode::Linear, 1.0f},
+	{EWaveEditorFadeMode::Exponetial, 3.0f},
+	{EWaveEditorFadeMode::Logarithmic, 0.25f},
+	{EWaveEditorFadeMode::Sigmoid, -0.1f}
+};
+
 namespace WaveformTransformationTrimFadeNames
 {
 	static FLazyName StartTimeName(GET_MEMBER_NAME_CHECKED(UWaveformTransformationTrimFade, StartTime));
 	static FLazyName EndTimeName(GET_MEMBER_NAME_CHECKED(UWaveformTransformationTrimFade, EndTime));
 }
 
-static void ApplyFadeIn(Audio::FAlignedFloatBuffer& InputAudio, const float FadeLength, const float FadeCurve, const int32 NumChannels, const float SampleRate)
+static void ApplyFadeIn(Audio::FAlignedFloatBuffer& InputAudio, const float FadeLength, const float FadeCurve, const float SCurveSharpness, const int32 NumChannels, const float SampleRate)
 {
 	check(NumChannels > 0);
 
-	if(InputAudio.Num() < NumChannels || FadeLength < SMALL_NUMBER)
+	if (InputAudio.Num() < NumChannels || FadeLength < SMALL_NUMBER)
 	{
 		return;
 	}
-	
-	const int32 FadeNumFrames = FMath::Min(FadeLength * SampleRate / NumChannels, InputAudio.Num() / NumChannels);
+
+	const int32 FadeNumFrames = FMath::Min(FadeLength * SampleRate, InputAudio.Num() / NumChannels);
 	float* InputPtr = InputAudio.GetData();
-	
-	for(int32 FrameIndex = 0;FrameIndex < FadeNumFrames; ++FrameIndex)
+
+	for (int32 FrameIndex = 0; FrameIndex < FadeNumFrames; ++FrameIndex)
 	{
 		const float FadeFraction = (float)FrameIndex / FadeNumFrames;
-		const float EnvValue = FMath::Pow(FadeFraction, FadeCurve);
-		
-		for(int32 ChannelIt = 0; ChannelIt < NumChannels; ++ChannelIt)
+		const float EnvValue = UWaveformTransformationTrimFade::GetFadeInCurveValue(FadeCurve, FadeFraction, SCurveSharpness);
+
+		for (int32 ChannelIt = 0; ChannelIt < NumChannels; ++ChannelIt)
 		{
 			*(InputPtr + ChannelIt) *= EnvValue;
 		}
@@ -35,25 +43,25 @@ static void ApplyFadeIn(Audio::FAlignedFloatBuffer& InputAudio, const float Fade
 	}
 }
 
-static void ApplyFadeOut(Audio::FAlignedFloatBuffer& InputAudio, const float FadeLength, const float FadeCurve, const int32 NumChannels, const float SampleRate)
+static void ApplyFadeOut(Audio::FAlignedFloatBuffer& InputAudio, const float FadeLength, const float FadeCurve, const float SCurveSharpness, const int32 NumChannels, const float SampleRate)
 {
 	check(NumChannels > 0);
 
-	if(InputAudio.Num() < NumChannels || FadeLength < SMALL_NUMBER)
+	if (InputAudio.Num() < NumChannels || FadeLength < SMALL_NUMBER)
 	{
 		return;
 	}
-	
-	const int32 FadeNumFrames = FMath::Min(FadeLength * SampleRate / NumChannels, InputAudio.Num() / NumChannels);
+
+	const int32 FadeNumFrames = FMath::Min(FadeLength * SampleRate, InputAudio.Num() / NumChannels);
 	const int32 StartSampleIndex = InputAudio.Num() - (FadeNumFrames * NumChannels);
 	float* InputPtr = &InputAudio[StartSampleIndex];
-	
-	for(int32 FrameIndex = 0;FrameIndex < FadeNumFrames; ++FrameIndex)
+
+	for (int32 FrameIndex = 0; FrameIndex < FadeNumFrames; ++FrameIndex)
 	{
 		const float FadeFraction = (float)FrameIndex / FadeNumFrames;
-		const float EnvValue = 1.f - FMath::Pow(FadeFraction, FadeCurve);
-		
-		for(int32 ChannelIt = 0; ChannelIt < NumChannels; ++ChannelIt)
+		float EnvValue = UWaveformTransformationTrimFade::GetFadeOutCurveValue(FadeCurve, FadeFraction, SCurveSharpness);
+
+		for (int32 ChannelIt = 0; ChannelIt < NumChannels; ++ChannelIt)
 		{
 			*(InputPtr + ChannelIt) *= EnvValue;
 		}
@@ -62,13 +70,15 @@ static void ApplyFadeOut(Audio::FAlignedFloatBuffer& InputAudio, const float Fad
 	}
 }
 
-FWaveTransformationTrimFade::FWaveTransformationTrimFade(double InStartTime, double InEndTime, float InStartFadeTime, float InStartFadeCurve, float InEndFadeTime, float InEndFadeCurve)
+FWaveTransformationTrimFade::FWaveTransformationTrimFade(double InStartTime, double InEndTime, float InStartFadeTime, float InStartFadeCurve, float InStartSCurveSharpness, float InEndFadeTime, float InEndFadeCurve, float InEndSCurveSharpness)
 	: StartTime(InStartTime)
 	, EndTime(InEndTime)
 	, StartFadeTime(InStartFadeTime)
-	, StartFadeCurve(FMath::Max(InStartFadeCurve, 0.f))
+	, StartFadeCurve(FMath::Max(InStartFadeCurve, -0.1f))
+	, StartSCurveSharpness(InStartSCurveSharpness)
 	, EndFadeTime(InEndFadeTime)
-	, EndFadeCurve(FMath::Max(InEndFadeCurve, 0.f)) {}
+	, EndFadeCurve(FMath::Max(InEndFadeCurve, -0.1f))
+	, EndSCurveSharpness(InEndSCurveSharpness){}
 
 void FWaveTransformationTrimFade::ProcessAudio(Audio::FWaveformTransformationWaveInfo& InOutWaveInfo) const
 {
@@ -121,18 +131,18 @@ void FWaveTransformationTrimFade::ProcessAudio(Audio::FWaveformTransformationWav
 
 	if(StartFadeTime > 0.f)
 	{
-		ApplyFadeIn(InputAudio, StartFadeTime, StartFadeCurve, InOutWaveInfo.NumChannels, InOutWaveInfo.SampleRate);
+		ApplyFadeIn(InputAudio, StartFadeTime, StartFadeCurve, StartSCurveSharpness, InOutWaveInfo.NumChannels, InOutWaveInfo.SampleRate);
 	}
 	
 	if(EndFadeTime > 0.f)
 	{
-		ApplyFadeOut(InputAudio, EndFadeTime, EndFadeCurve, InOutWaveInfo.NumChannels, InOutWaveInfo.SampleRate);
+		ApplyFadeOut(InputAudio, EndFadeTime, EndFadeCurve, EndSCurveSharpness, InOutWaveInfo.NumChannels, InOutWaveInfo.SampleRate);
 	}
 }
 
 Audio::FTransformationPtr UWaveformTransformationTrimFade::CreateTransformation() const
 {
-	return MakeUnique<FWaveTransformationTrimFade>(StartTime, EndTime, StartFadeTime, StartFadeCurve, EndFadeTime, EndFadeCurve);
+	return MakeUnique<FWaveTransformationTrimFade>(StartTime, EndTime, StartFadeTime, StartFadeCurve, StartSCurveSharpness, EndFadeTime, EndFadeCurve, EndSCurveSharpness);
 }
 
 void UWaveformTransformationTrimFade::UpdateConfiguration(FWaveTransformUObjectConfiguration& InOutConfiguration)
@@ -143,6 +153,58 @@ void UWaveformTransformationTrimFade::UpdateConfiguration(FWaveTransformUObjectC
 	InOutConfiguration.EndTime = EndTime;
 }
 
+const double UWaveformTransformationTrimFade::GetFadeInCurveValue(const float StartFadeCurve, const double FadeFraction, const float SCurveSharpness)
+{
+	double CurveValue = 1;
+
+	if (StartFadeCurve < 0)
+	{
+		double Slope = 10 * SCurveSharpness * (PI + 1);
+		
+		if (FadeFraction <= 0.5)
+		{
+			CurveValue = (((FMath::Tanh(((Slope * FadeFraction * PI) / 2) - ((Slope * PI) / 4))) + 1) / 2) * (2 * FadeFraction);
+		}
+		else
+		{
+			CurveValue = ((FMath::Tanh((-Slope * (-FadeFraction + 1) * PI) / 2 + ((Slope * PI) / 4)) - 1) / 2) * (2 * (-FadeFraction + 1)) + 1;
+		}
+	}
+	else
+	{
+		CurveValue = FMath::Pow(FadeFraction, StartFadeCurve);
+	}
+
+	return CurveValue;
+}
+
+const double UWaveformTransformationTrimFade::GetFadeOutCurveValue(const float EndFadeCurve, const double FadeFraction, const float SCurveSharpness)
+{
+	double CurveValue = 1;
+
+	if (EndFadeCurve < 0)
+	{
+		double Slope = 10 * SCurveSharpness * (PI + 1);
+		
+		if (FadeFraction <= 0.5)
+		{
+			CurveValue = -(((FMath::Tanh(((Slope * FadeFraction * PI) / 2) - ((Slope * PI) / 4))) + 1) / 2) * (2 * FadeFraction) + 1;
+		}
+		else
+		{
+			CurveValue = -((FMath::Tanh((-Slope * (-FadeFraction + 1) * PI) / 2 + ((Slope * PI) / 4)) - 1) / 2) * (2 * (-FadeFraction + 1));
+		}
+	}
+	else
+	{
+		// Changed from 1.f - FMath::Pow(FadeFraction, FadeCurve) so that the fade out curve is a horizontally mirrored version of fade in 
+		// instead of vertically mirrored
+		CurveValue = FMath::Pow(-FadeFraction + 1, EndFadeCurve);
+	}
+
+	return CurveValue;
+}
+
 void UWaveformTransformationTrimFade::UpdateDurationProperties(const float InAvailableDuration)
 {
 	check(InAvailableDuration > 0)
diff --git a/Engine/Plugins/Editor/WaveformEditor/Source/WaveformTransformations/Public/WaveformTransformationTrimFade.h b/Engine/Plugins/Editor/WaveformEditor/Source/WaveformTransformations/Public/WaveformTransformationTrimFade.h
--- a/Engine/Plugins/Editor/WaveformEditor/Source/WaveformTransformations/Public/WaveformTransformationTrimFade.h
+++ b/Engine/Plugins/Editor/WaveformEditor/Source/WaveformTransformations/Public/WaveformTransformationTrimFade.h
@@ -7,10 +7,19 @@
 
 #include "WaveformTransformationTrimFade.generated.h"
 
+UENUM()
+enum class EWaveEditorFadeMode : uint8
+{
+	Linear = 0,
+	Exponetial,
+	Logarithmic,
+	Sigmoid
+};
+
 class WAVEFORMTRANSFORMATIONS_API FWaveTransformationTrimFade : public Audio::IWaveTransformation
 {
 public:
-	explicit FWaveTransformationTrimFade(double InStartTime, double InEndTime, float InStartFadeTime, float InStartFadeCurve, float InEndFadeTime, float InEndFadeCurve);
+	explicit FWaveTransformationTrimFade(double InStartTime, double InEndTime, float InStartFadeTime, float InStartFadeCurve, float InStartSCurveSharpness, float InEndFadeTime, float InEndFadeCurve, float InEndSCurveSharpness);
 	virtual void ProcessAudio(Audio::FWaveformTransformationWaveInfo& InOutWaveInfo) const override;
 
 	virtual bool CanChangeFileLength() const override { return true; }
@@ -21,9 +30,11 @@ private:
 
 	float StartFadeTime = 0.f;
 	float StartFadeCurve = 0.f;
+	float StartSCurveSharpness = 0.f;
 
 	float EndFadeTime = 0.f;
 	float EndFadeCurve = 0.f;
+	float EndSCurveSharpness = 0.f;
 };
 
 UCLASS()
@@ -41,18 +52,29 @@ public:
 	UPROPERTY(EditAnywhere, Category = "Fade", meta=(ClampMin = 0.0, DisplayName = "Fade-In Duration"))
 	float StartFadeTime = 0.f;
 
-	UPROPERTY(EditAnywhere, Category = "Fade", meta=(ClampMin = 0.0, ClampMax = 10.0, DisplayName = "Fade-In Curve"))
+	UPROPERTY(EditAnywhere, Category = "Fade", meta=(ClampMin = -0.1, ClampMax = 10.0, DisplayName = "Fade-In Curve"))
 	float StartFadeCurve = 1.f;
 
+	UPROPERTY(EditAnywhere, Category = "Fade", meta = (ClampMin = -1.0, ClampMax = 1.0, DisplayName = "S-Curve Sharpness", EditCondition = "StartFadeCurve < 0", EditConditionHides))
+	float StartSCurveSharpness = 0.1f;
+
 	UPROPERTY(EditAnywhere, Category = "Fade", meta=(ClampMin = 0.0, DisplayName = "Fade-Out Duration"))
 	float EndFadeTime = 0.f;
 
-	UPROPERTY(EditAnywhere, Category = "Fade", meta=(ClampMin = 0.0, ClampMax = 10.0, DisplayName = "Fade-Out Curve"))
+	UPROPERTY(EditAnywhere, Category = "Fade", meta=(ClampMin = -0.1, ClampMax = 10.0, DisplayName = "Fade-Out Curve"))
 	float EndFadeCurve = 1.f;
 
+	UPROPERTY(EditAnywhere, Category = "Fade", meta = (ClampMin = -1.0, ClampMax = 1.0, DisplayName = "S-Curve Sharpness", EditCondition = "EndFadeCurve < 0", EditConditionHides))
+	float EndSCurveSharpness = 0.1;
+	
+	static const TMap<EWaveEditorFadeMode, float> FadeModeToCurveValueMap;
+
 	virtual Audio::FTransformationPtr CreateTransformation() const override;
 
 	void UpdateConfiguration(FWaveTransformUObjectConfiguration& InOutConfiguration) override;
+
+	static const double GetFadeInCurveValue(const float StartFadeCurve, const double FadeFraction, const float SCurveSharpness = 0);
+	static const double GetFadeOutCurveValue(const float EndFadeCurve, const double FadeFraction, const float SCurveSharpness = 0);
 	
 private:	
 	void UpdateDurationProperties(const float InAvailableDuration);
diff --git a/Engine/Plugins/Editor/WaveformEditor/Source/WaveformTransformationsWidgets/Private/WaveformTransformationTrimFadeRenderer.cpp b/Engine/Plugins/Editor/WaveformEditor/Source/WaveformTransformationsWidgets/Private/WaveformTransformationTrimFadeRenderer.cpp
--- a/Engine/Plugins/Editor/WaveformEditor/Source/WaveformTransformationsWidgets/Private/WaveformTransformationTrimFadeRenderer.cpp
+++ b/Engine/Plugins/Editor/WaveformEditor/Source/WaveformTransformationsWidgets/Private/WaveformTransformationTrimFadeRenderer.cpp
@@ -2,10 +2,13 @@
 
 #include "WaveformTransformationTrimFadeRenderer.h"
 
+#include "Brushes/SlateRoundedBoxBrush.h"
+#include "Framework/Application/SlateApplication.h"
 #include "Layout/Geometry.h"
 #include "PropertyHandle.h"
 #include "Rendering/DrawElements.h"
 #include "WaveformTransformationTrimFade.h"
+#include "Widgets/Input/SButton.h"
 #include "Widgets/SWidget.h"
 
 int32 FWaveformTransformationTrimFadeRenderer::OnPaint(const FPaintArgs& Args, const FGeometry& AllottedGeometry, const FSlateRect& MyCullingRect, FSlateWindowElementList& OutDrawElements, int32 LayerId, const FWidgetStyle& InWidgetStyle, bool bParentEnabled) const
@@ -61,6 +64,11 @@ int32 FWaveformTransformationTrimFadeRenderer::DrawTrimHandles(const FGeometry&
 
 int32 FWaveformTransformationTrimFadeRenderer::DrawFadeCurves(const FGeometry& AllottedGeometry, FSlateWindowElementList& OutDrawElements, int32 LayerId) const
 {
+	const FLinearColor FadeColor = FLinearColor::Gray;
+	const FSlateRoundedBoxBrush RoundedBoxBrush(FadeColor, 2);
+	const float HandleWidth = InteractionPixelXDelta;
+	const float HandleHeight = AllottedGeometry.GetLocalSize().Y * InteractionRatioYDelta;
+
 	if (FadeInCurvePoints.Num() > 0)
 	{
 		FSlateDrawElement::MakeLines(
@@ -69,9 +77,20 @@ int32 FWaveformTransformationTrimFadeRenderer::DrawFadeCurves(const FGeometry& A
 			AllottedGeometry.ToPaintGeometry(),
 			FadeInCurvePoints,
 			ESlateDrawEffect::None,
-			FLinearColor::Yellow
+			FadeColor
 		);
 
+		const float	HandleStart = FadeInCurvePoints.Last(0).X;
+		const FPaintGeometry HandleGeometry = AllottedGeometry.ToPaintGeometry(FVector2f(HandleWidth, HandleHeight), FSlateLayoutTransform(FVector2f(HandleStart, 0)));
+		
+		FSlateDrawElement::MakeBox(
+			OutDrawElements,
+			++LayerId,
+			HandleGeometry,
+			&RoundedBoxBrush,
+			ESlateDrawEffect::None,
+			FadeColor
+		);
 	}
 
 	if (FadeOutCurvePoints.Num() > 0)
@@ -82,9 +101,20 @@ int32 FWaveformTransformationTrimFadeRenderer::DrawFadeCurves(const FGeometry& A
 			AllottedGeometry.ToPaintGeometry(),
 			FadeOutCurvePoints,
 			ESlateDrawEffect::None,
-			FLinearColor::Yellow
+			FadeColor
 		);
 
+		const float	HandleStart = FadeOutCurvePoints[0].X - HandleWidth;
+		const FPaintGeometry HandleGeometry = AllottedGeometry.ToPaintGeometry(FVector2f(HandleWidth, HandleHeight), FSlateLayoutTransform(FVector2f(HandleStart, 0)));
+
+		FSlateDrawElement::MakeBox(
+			OutDrawElements,
+			++LayerId,
+			HandleGeometry,
+			&RoundedBoxBrush,
+			ESlateDrawEffect::None,
+			FadeColor
+		);
 	}
 
 	return LayerId;
@@ -93,17 +123,18 @@ int32 FWaveformTransformationTrimFadeRenderer::DrawFadeCurves(const FGeometry& A
 void FWaveformTransformationTrimFadeRenderer::GenerateFadeCurves(const FGeometry& AllottedGeometry)
 {
 	const float FadeInFrames = GetPropertyValue<float>(StartFadeTimeHandle) * TransformationWaveInfo.SampleRate;
-	const uint32 FadeInPixelLenght = FadeInFrames * PixelsPerFrame;
+	const uint32 FadeInPixelLength = FadeInFrames * PixelsPerFrame;
 	FadeInStartX = FMath::RoundToInt32(StartTimeHandleX);
-	FadeInEndX = FMath::RoundToInt32(FMath::Clamp(StartTimeHandleX + FadeInPixelLenght, StartTimeHandleX, EndTimeHandleX));
+	FadeInEndX = FMath::RoundToInt32(FMath::Clamp(StartTimeHandleX + FadeInPixelLength, StartTimeHandleX, EndTimeHandleX));
 	
-	const uint32 DisplayedFadeInPixelLenght = FadeInEndX - FadeInStartX;
-	FadeInCurvePoints.SetNumUninitialized(DisplayedFadeInPixelLenght);
+	const uint32 DisplayedFadeInPixelLength = FadeInEndX - FadeInStartX;
+	FadeInCurvePoints.SetNumUninitialized(DisplayedFadeInPixelLength);
 
-	for (uint32 Pixel = 0; Pixel < DisplayedFadeInPixelLenght; ++Pixel)
+	for (uint32 Pixel = 0; Pixel < DisplayedFadeInPixelLength; ++Pixel)
 	{
-		const double FadeFraction = (float)Pixel / FadeInPixelLenght;
-		const double CurveValue = Pixel != FadeInPixelLenght - 1 ? 1.f - FMath::Pow(FadeFraction, GetPropertyValue<float>(StartFadeCurveHandle)) : 0.f;
+		const double FadeFraction = (float)Pixel / FadeInPixelLength;
+		const double CurveFunction = UWaveformTransformationTrimFade::GetFadeInCurveValue(GetPropertyValue<float>(StartFadeCurveHandle), FadeFraction, GetPropertyValue<float>(StartSCurveSharpnessHandle));
+		const double CurveValue = Pixel != FadeInPixelLength - 1 ? 1.f - CurveFunction : 0.f;
 
 		const uint32 XCoordinate = Pixel + FadeInStartX;
 		FadeInCurvePoints[Pixel] = FVector2D(XCoordinate, CurveValue * AllottedGeometry.Size.Y);
@@ -121,7 +152,8 @@ void FWaveformTransformationTrimFadeRenderer::GenerateFadeCurves(const FGeometry
 	for (uint32 Pixel = 0; Pixel < DisplayedFadeOutPixelLength; ++Pixel)
 	{
 		const double FadeFraction = (float)(Pixel + FadeOutPixelOffset) / FadeOutPixelLength;
-		const double CurveValue = Pixel != FadeOutPixelLength - 1 ? FMath::Pow(FadeFraction, GetPropertyValue<float>(EndFadeCurveHandle)) : 1.f;
+		const double CurveFunction = UWaveformTransformationTrimFade::GetFadeOutCurveValue(GetPropertyValue<float>(EndFadeCurveHandle), FadeFraction, GetPropertyValue<float>(EndSCurveSharpnessHandle));
+		const double CurveValue = Pixel != FadeOutPixelLength - 1 ? 1.f - CurveFunction : 1.f;
 
 		const uint32 XCoordinate = Pixel + FadeOutStartX;
 		FadeOutCurvePoints[Pixel] = FVector2D(XCoordinate, CurveValue * AllottedGeometry.Size.Y);
@@ -158,8 +190,10 @@ void FWaveformTransformationTrimFadeRenderer::SetPropertyHandles(const TArray<TS
 	EndTimeHandle  = GetPropertyHandle(GET_MEMBER_NAME_CHECKED(UWaveformTransformationTrimFade, EndTime));
 	StartFadeTimeHandle = GetPropertyHandle(GET_MEMBER_NAME_CHECKED(UWaveformTransformationTrimFade, StartFadeTime));
 	StartFadeCurveHandle = GetPropertyHandle(GET_MEMBER_NAME_CHECKED(UWaveformTransformationTrimFade, StartFadeCurve));
+	StartSCurveSharpnessHandle = GetPropertyHandle(GET_MEMBER_NAME_CHECKED(UWaveformTransformationTrimFade, StartSCurveSharpness));
 	EndFadeTimeHandle = GetPropertyHandle(GET_MEMBER_NAME_CHECKED(UWaveformTransformationTrimFade, EndFadeTime));
 	EndFadeCurveHandle = GetPropertyHandle(GET_MEMBER_NAME_CHECKED(UWaveformTransformationTrimFade, EndFadeCurve));
+	EndSCurveSharpnessHandle = GetPropertyHandle(GET_MEMBER_NAME_CHECKED(UWaveformTransformationTrimFade, EndSCurveSharpness));
 }
 
 FReply FWaveformTransformationTrimFadeRenderer::OnMouseWheel(SWidget& OwnerWidget, const FGeometry& MyGeometry, const FPointerEvent& MouseEvent)
@@ -235,18 +269,14 @@ FReply FWaveformTransformationTrimFadeRenderer::OnMouseButtonDown(SWidget& Owner
 {
 	const FVector2D LocalCursorPosition = GetLocalCursorPosition(MouseEvent, MyGeometry);
 
-	const bool bIsLeftMouseButton = MouseEvent.IsMouseButtonDown(EKeys::LeftMouseButton);
+	const FKey MouseButton = MouseEvent.GetEffectingButton();
 
-	if (bIsLeftMouseButton)
-	{
-		TrimFadeInteractionType = GetInteractionTypeFromCursorPosition(LocalCursorPosition, MyGeometry);
-
-		if (TrimFadeInteractionType != ETrimFadeInteractionType::None)
-		{
-			BeginTransaction(TEXT("PropertyEditor"), FText::Format(NSLOCTEXT("PropertyEditor", "SetNumericPropertyTransaction", "Edit {0}"), GetPropertyEditedByCurrentInteraction()), nullptr);
-			return FReply::Handled().CaptureMouse(OwnerWidget.AsShared()).PreventThrottling();
-		}
+	TrimFadeInteractionType = GetInteractionTypeFromCursorPosition(LocalCursorPosition, MouseButton, MyGeometry);
 
+	if (TrimFadeInteractionType != ETrimFadeInteractionType::None)
+	{
+		BeginTransaction(TEXT("PropertyEditor"), FText::Format(NSLOCTEXT("PropertyEditor", "SetNumericPropertyTransaction", "Edit {0}"), GetPropertyEditedByCurrentInteraction()), nullptr);
+		return FReply::Handled().CaptureMouse(OwnerWidget.AsShared()).PreventThrottling();
 	}
 
 	return FReply::Unhandled();
@@ -308,31 +338,100 @@ void FWaveformTransformationTrimFadeRenderer::SetPropertyValueDependingOnInterac
 			EndFadeTimeHandle->SetValue(EndFadeTimeValue, Flags);
 		}
 		break;
+	case FWaveformTransformationTrimFadeRenderer::ETrimFadeInteractionType::RightClickFadeIn:
+		ShowSelectFadeModeMenuAtCursor(WidgetGeometry, MouseEvent, StartFadeCurveHandle);
+		break;
+	case FWaveformTransformationTrimFadeRenderer::ETrimFadeInteractionType::RightClickFadeOut:
+		ShowSelectFadeModeMenuAtCursor(WidgetGeometry, MouseEvent, EndFadeCurveHandle);
+		break;
 	default:
 		break;
 	}
 }
 
-FWaveformTransformationTrimFadeRenderer::ETrimFadeInteractionType FWaveformTransformationTrimFadeRenderer::GetInteractionTypeFromCursorPosition(const FVector2D& InLocalCursorPosition, const FGeometry& WidgetGeometry) const
+void FWaveformTransformationTrimFadeRenderer::ShowSelectFadeModeMenuAtCursor(const FGeometry& MyGeometry, const FPointerEvent& MouseEvent, TSharedPtr<IPropertyHandle> PropertyHandle)
 {
-	if (IsCursorInFadeInInteractionRange(InLocalCursorPosition, WidgetGeometry))
+	if (!PropertyHandle.IsValid())
 	{
-		return ETrimFadeInteractionType::ScrubbingFadeIn;
+		return;
 	}
 
-	if (IsCursorInFadeOutInteractionRange(InLocalCursorPosition, WidgetGeometry))
+	const FVector2D LocalCursorPosition = MouseEvent.GetScreenSpacePosition();
+
+	TSharedRef<SVerticalBox> MenuContent = SNew(SVerticalBox);
+
+	for (const TPair<EWaveEditorFadeMode, float>& FadeOptionPair : UWaveformTransformationTrimFade::FadeModeToCurveValueMap)
 	{
-		return ETrimFadeInteractionType::ScrubbingFadeOut;
+		MenuContent->AddSlot()
+			.Padding(5)
+			.AutoHeight()
+			[
+				SNew(SButton)
+					.Text(FText::FromString(StaticEnum<EWaveEditorFadeMode>()->GetNameStringByValue(static_cast<int64>(FadeOptionPair.Key))))
+					.OnClicked_Lambda([this, PropertyHandle, FadeOptionPair]() -> FReply {
+					PropertyHandle->SetValue(FadeOptionPair.Value);
+
+					if (TSharedPtr<SWindow> LockedFadeModeMenuWindow = FadeModeMenuWindow.Pin())
+					{
+						FSlateApplication::Get().RequestDestroyWindow(LockedFadeModeMenuWindow.ToSharedRef());
+					}
+
+					return FReply::Handled();
+				})
+			];
 	}
 
-	if (StartTimeInteractionXRange.Contains(InLocalCursorPosition.X))
+	TSharedRef<SWindow> MenuWindow = SNew(SWindow)
+		.AutoCenter(EAutoCenter::None)
+		.ScreenPosition(LocalCursorPosition)
+		.SizingRule(ESizingRule::Autosized)
+		.SupportsTransparency(EWindowTransparency::None)
+		.IsPopupWindow(true)
+		.CreateTitleBar(false)
+		[
+			MenuContent
+		];
+
+	FadeModeMenuWindow = MenuWindow;
+
+	FSlateApplication::Get().AddWindow(MenuWindow);
+}
+
+FWaveformTransformationTrimFadeRenderer::ETrimFadeInteractionType FWaveformTransformationTrimFadeRenderer::GetInteractionTypeFromCursorPosition(const FVector2D& InLocalCursorPosition, const FKey MouseButton, const FGeometry& WidgetGeometry) const
+{
+	if (MouseButton == EKeys::LeftMouseButton)
 	{
-		return ETrimFadeInteractionType::ScrubbingLeftHandle;
-	}
+		if (IsCursorInFadeInInteractionRange(InLocalCursorPosition, WidgetGeometry))
+		{
+			return ETrimFadeInteractionType::ScrubbingFadeIn;
+		}
 
-	if (EndTimeInteractionXRange.Contains(InLocalCursorPosition.X))
+		if (IsCursorInFadeOutInteractionRange(InLocalCursorPosition, WidgetGeometry))
+		{
+			return ETrimFadeInteractionType::ScrubbingFadeOut;
+		}
+
+		if (StartTimeInteractionXRange.Contains(InLocalCursorPosition.X))
+		{
+			return ETrimFadeInteractionType::ScrubbingLeftHandle;
+		}
+
+		if (EndTimeInteractionXRange.Contains(InLocalCursorPosition.X))
+		{
+			return ETrimFadeInteractionType::ScrubbingRightHandle;
+		}
+	}
+	else if (MouseButton == EKeys::RightMouseButton)
 	{
-		return ETrimFadeInteractionType::ScrubbingRightHandle;
+		if (IsCursorInFadeInInteractionRange(InLocalCursorPosition, WidgetGeometry))
+		{
+			return ETrimFadeInteractionType::RightClickFadeIn;
+		}
+
+		if (IsCursorInFadeOutInteractionRange(InLocalCursorPosition, WidgetGeometry))
+		{
+			return ETrimFadeInteractionType::RightClickFadeOut;
+		}
 	}
 
 	return ETrimFadeInteractionType::None;
diff --git a/Engine/Plugins/Editor/WaveformEditor/Source/WaveformTransformationsWidgets/Private/WaveformTransformationTrimFadeRenderer.h b/Engine/Plugins/Editor/WaveformEditor/Source/WaveformTransformationsWidgets/Private/WaveformTransformationTrimFadeRenderer.h
--- a/Engine/Plugins/Editor/WaveformEditor/Source/WaveformTransformationsWidgets/Private/WaveformTransformationTrimFadeRenderer.h
+++ b/Engine/Plugins/Editor/WaveformEditor/Source/WaveformTransformationsWidgets/Private/WaveformTransformationTrimFadeRenderer.h
@@ -54,16 +54,23 @@ private:
 		ScrubbingLeftHandle,
 		ScrubbingRightHandle,
 		ScrubbingFadeIn, 
-		ScrubbingFadeOut
+		ScrubbingFadeOut,
+		RightClickFadeIn,
+		RightClickFadeOut
 	} TrimFadeInteractionType = ETrimFadeInteractionType::None;
 
-	ETrimFadeInteractionType GetInteractionTypeFromCursorPosition(const FVector2D& InLocalCursorPosition, const FGeometry& WidgetGeometry) const;
+	ETrimFadeInteractionType GetInteractionTypeFromCursorPosition(const FVector2D& InLocalCursorPosition, const FKey MouseButton, const FGeometry& WidgetGeometry) const;
 	void SetPropertyValueDependingOnInteractionType(const FPointerEvent& MouseEvent, const FGeometry& WidgetGeometry, const EPropertyValueSetFlags::Type Flags = EPropertyValueSetFlags::DefaultFlags);
+	void ShowSelectFadeModeMenuAtCursor(const FGeometry& MyGeometry, const FPointerEvent& MouseEvent, TSharedPtr<IPropertyHandle> PropertyHandle);
+
+	TWeakPtr<SWindow> FadeModeMenuWindow;
 
 	TSharedPtr<IPropertyHandle> StartTimeHandle;
 	TSharedPtr<IPropertyHandle> EndTimeHandle;
 	TSharedPtr<IPropertyHandle> StartFadeTimeHandle;
 	TSharedPtr<IPropertyHandle> StartFadeCurveHandle;
+	TSharedPtr<IPropertyHandle> StartSCurveSharpnessHandle;
 	TSharedPtr<IPropertyHandle> EndFadeTimeHandle;
 	TSharedPtr<IPropertyHandle> EndFadeCurveHandle;
+	TSharedPtr<IPropertyHandle> EndSCurveSharpnessHandle;
 };
