Warning fix for:  Incremental Static Analysis Win64 (PVS-Studio): \WaveformEditor.cpp(210): warning V674: The '0.5f' constant of the 'float' type is compared to a value of the 'int' type. Consider inspecting the 'DefaultFadeInAmount >= 0' expression.

diff --git a/Engine/Plugins/Editor/WaveformEditor/Source/WaveformEditor/Private/WaveformEditor.cpp b/Engine/Plugins/Editor/WaveformEditor/Source/WaveformEditor/Private/WaveformEditor.cpp
--- a/Engine/Plugins/Editor/WaveformEditor/Source/WaveformEditor/Private/WaveformEditor.cpp
+++ b/Engine/Plugins/Editor/WaveformEditor/Source/WaveformEditor/Private/WaveformEditor.cpp
@@ -207,10 +207,10 @@ void FWaveformEditor::ToggleFadeIn()
 	//if fade duration is 0, assume there is no fade present and toggle on the fade with the curve data selected in the menu
 	if (TrimFadeTransformation->StartFadeTime == 0.0f)
 	{
-		check(DefaultFadeInAmount >= 0);
+		check(DefaultFadeInAmount >= 0.0f);
 		check(FadeInCurve >= -0.1f && FadeInCurve <= 10);
 
-		if (CachedFadeInAmount <= 0)
+		if (CachedFadeInAmount <= 0.0f)
 		{
 			CachedFadeInAmount = DefaultFadeInAmount;
 		}
@@ -246,10 +246,10 @@ void FWaveformEditor::ToggleFadeOut()
 	//if fade duration is 0, assume there is no fade present and toggle on the fade with the curve data selected in the menu
 	if (TrimFadeTransformation->EndFadeTime == 0.0f)
 	{
-		check(DefaultFadeOutAmount >= 0);
+		check(DefaultFadeOutAmount >= 0.0f);
 		check(FadeOutCurve >= -0.1f && FadeOutCurve <= 10);
 
-		if (CachedFadeOutAmount <= 0)
+		if (CachedFadeOutAmount <= 0.0f)
 		{
 			CachedFadeOutAmount = DefaultFadeOutAmount;
 		}
