[Waveform Editor] - Selecting Cancel when exporting sound wave does not keep asset

diff --git a/Engine/Plugins/Editor/WaveformEditor/Source/WaveformEditor/Private/WaveformEditorWaveWriter.cpp b/Engine/Plugins/Editor/WaveformEditor/Source/WaveformEditor/Private/WaveformEditorWaveWriter.cpp
--- a/Engine/Plugins/Editor/WaveformEditor/Source/WaveformEditor/Private/WaveformEditorWaveWriter.cpp
+++ b/Engine/Plugins/Editor/WaveformEditor/Source/WaveformEditor/Private/WaveformEditorWaveWriter.cpp
@@ -4,6 +4,7 @@
 
 #include "AssetToolsModule.h"
 #include "DSP/FloatArrayMath.h"
+#include "EditorAssetLibrary.h"
 #include "FileHelpers.h"
 #include "Sound/SoundWave.h"
 #include "WaveformEditorLog.h"
@@ -47,7 +48,13 @@ void FWaveformEditorWaveWriter::ExportTransformedWaveform()
 		if(ResultingWave->GetPackage())
 		{
 			PackagesToSave.Add(ResultingWave->GetPackage());
-			FEditorFileUtils::PromptForCheckoutAndSave(PackagesToSave, false /*bCheckDirty*/, true /*bPromptToSave*/);
+			FEditorFileUtils::EPromptReturnCode Result = FEditorFileUtils::PromptForCheckoutAndSave(PackagesToSave, false /*bCheckDirty*/, true /*bPromptToSave*/);
+
+			// If the user cancels the checkout and save, delete the new generated soundwave
+			if (Result == FEditorFileUtils::PR_Cancelled)
+			{
+				UEditorAssetLibrary::DeleteAsset(ResultingWave->GetPathName());
+			}
 		}
 	};
 
diff --git a/Engine/Plugins/Editor/WaveformEditor/Source/WaveformEditor/WaveformEditor.build.cs b/Engine/Plugins/Editor/WaveformEditor/Source/WaveformEditor/WaveformEditor.build.cs
--- a/Engine/Plugins/Editor/WaveformEditor/Source/WaveformEditor/WaveformEditor.build.cs
+++ b/Engine/Plugins/Editor/WaveformEditor/Source/WaveformEditor/WaveformEditor.build.cs
@@ -20,6 +20,7 @@ namespace UnrealBuildTool.Rules
 					"CoreUObject",
 					"ContentBrowser",
 					"DeveloperSettings",
+					"EditorScriptingUtilities",
 					"EditorStyle",
 					"Engine",
 					"InputCore",
diff --git a/Engine/Plugins/Editor/WaveformEditor/WaveformEditor.uplugin b/Engine/Plugins/Editor/WaveformEditor/WaveformEditor.uplugin
--- a/Engine/Plugins/Editor/WaveformEditor/WaveformEditor.uplugin
+++ b/Engine/Plugins/Editor/WaveformEditor/WaveformEditor.uplugin
@@ -44,6 +44,10 @@
 		{
 			"Name": "AudioWidgets",
 			"Enabled" :  true
+		},
+		{
+			"Name": "EditorScriptingUtilities",
+			"Enabled": true
 		}
 	]
 }
\ No newline at end of file
