Waveform UI: improve UI feedback for resizing marker/loop points
Waveform editor: make import/overwrite button+ellipses work like the platform UI
Waveform Marker Transformation Renderer Bug fix: editing markers with the UI now marks the soundwave dirty

diff --git a/Engine/Plugins/Editor/WaveformEditor/Source/WaveformEditor/Private/WaveformEditor.cpp b/Engine/Plugins/Editor/WaveformEditor/Source/WaveformEditor/Private/WaveformEditor.cpp
--- a/Engine/Plugins/Editor/WaveformEditor/Source/WaveformEditor/Private/WaveformEditor.cpp
+++ b/Engine/Plugins/Editor/WaveformEditor/Source/WaveformEditor/Private/WaveformEditor.cpp
@@ -360,6 +360,11 @@ void FWaveformEditor::ExecuteReimport()
 	{
 		return;
 	}
+	
+	if (ReimportMode == EWaveEditorReimportMode::SameFileOverwrite)
+	{
+		ExecuteOverwriteTransformations();
+	}
 
 	const bool bSelectNewAsset = ReimportMode == EWaveEditorReimportMode::SelectFile;
 
@@ -368,14 +373,12 @@ void FWaveformEditor::ExecuteReimport()
 	FReimportManager::Instance()->ValidateAllSourceFileAndReimport(CopyOfSelectedAssets, true, -1, bSelectNewAsset);
 }
 
-void FWaveformEditor::ExecuteReimportAndOverwriteTransformations()
+void FWaveformEditor::ExecuteOverwriteTransformations()
 {
 	for (TObjectPtr<UWaveformTransformationBase> Transformation : SoundWave->Transformations)
 	{
 		Transformation->OverwriteTransformation();
 	}
-
-	ExecuteReimport();
 }
 
 void FWaveformEditor::RegisterTabSpawners(const TSharedRef<FTabManager>& InTabManager)
@@ -775,19 +778,19 @@ bool FWaveformEditor::BindCommands()
 
 	ToolkitCommands->MapAction(
 		Commands.ReimportModeSameFile,
-		FExecuteAction::CreateLambda([this] { ReimportMode = EWaveEditorReimportMode::SameFile; ExecuteReimport(); }),
+		FExecuteAction::CreateLambda([this] { ReimportMode = EWaveEditorReimportMode::SameFile; }),
 		FCanExecuteAction::CreateLambda([this] { return CanExecuteReimport(); }),
 		FIsActionChecked::CreateLambda([this] {return ReimportMode == EWaveEditorReimportMode::SameFile; }));
 
 	ToolkitCommands->MapAction(
 		Commands.ReimportModeSameFileOverwriteTransformations,
-		FExecuteAction::CreateLambda([this] { ReimportMode = EWaveEditorReimportMode::SameFile; ExecuteReimportAndOverwriteTransformations(); }),
+		FExecuteAction::CreateLambda([this] { ReimportMode = EWaveEditorReimportMode::SameFileOverwrite; }),
 		FCanExecuteAction::CreateLambda([this] { return CanExecuteReimport(); }),
-		FIsActionChecked::CreateLambda([this] {return ReimportMode == EWaveEditorReimportMode::SameFile; }));
+		FIsActionChecked::CreateLambda([this] {return ReimportMode == EWaveEditorReimportMode::SameFileOverwrite; }));
 
 	ToolkitCommands->MapAction(
 		Commands.ReimportModeNewFile,
-		FExecuteAction::CreateLambda([this] { ReimportMode = EWaveEditorReimportMode::SelectFile; ExecuteReimport(); }),
+		FExecuteAction::CreateLambda([this] { ReimportMode = EWaveEditorReimportMode::SelectFile; }),
 		FCanExecuteAction::CreateLambda([this] { return CanExecuteReimport(); }),
 		FIsActionChecked::CreateLambda([this] {return ReimportMode == EWaveEditorReimportMode::SelectFile; }));
 
@@ -1297,8 +1300,11 @@ FText FWaveformEditor::GetReimportButtonToolTip() const
 	case(EWaveEditorReimportMode::SameFile):
 		ReimportModeText = LOCTEXT("SameFile", "Reimport from same file");
 		break;
+	case(EWaveEditorReimportMode::SameFileOverwrite):
+		ReimportModeText = LOCTEXT("SameFileOverwrite", "Reimport from same file and overwrite transformations");
+		break;
 	default:
-		static_assert(static_cast<int32>(EWaveEditorReimportMode::COUNT) == 2, "Possible missing switch case coverage for 'EWaveEditorReimportMode'");
+		static_assert(static_cast<int32>(EWaveEditorReimportMode::COUNT) == 3, "Possible missing switch case coverage for 'EWaveEditorReimportMode'");
 		break;
 	}
 
diff --git a/Engine/Plugins/Editor/WaveformEditor/Source/WaveformEditor/Private/WaveformEditorCommands.cpp b/Engine/Plugins/Editor/WaveformEditor/Source/WaveformEditor/Private/WaveformEditorCommands.cpp
--- a/Engine/Plugins/Editor/WaveformEditor/Source/WaveformEditor/Private/WaveformEditorCommands.cpp
+++ b/Engine/Plugins/Editor/WaveformEditor/Source/WaveformEditor/Private/WaveformEditorCommands.cpp
@@ -28,9 +28,9 @@ void FWaveformEditorCommands::RegisterCommands()
 	UI_COMMAND(ExportFormatStereo, "Stereo", "Sets the export format to stereo", EUserInterfaceActionType::RadioButton, FInputChord());
 
 	UI_COMMAND(ReimportAsset, "Reimport Asset", "Reimports the inspected sound wave uasset", EUserInterfaceActionType::Button, FInputChord());
-	UI_COMMAND(ReimportModeSameFile, "Reimport from same file", "Reimports asset from the same file", EUserInterfaceActionType::Check, FInputChord());
-	UI_COMMAND(ReimportModeSameFileOverwriteTransformations, "Reimport from same file and overwrite transformations", "Reimports asset from the same file and overwrites the transformations", EUserInterfaceActionType::Check, FInputChord());
-	UI_COMMAND(ReimportModeNewFile, "Reimport from new file", "Reimports asset from a selected file", EUserInterfaceActionType::Check, FInputChord());
+	UI_COMMAND(ReimportModeSameFile, "Reimport from same file", "Reimports asset from the same file", EUserInterfaceActionType::RadioButton, FInputChord());
+	UI_COMMAND(ReimportModeSameFileOverwriteTransformations, "Reimport from same file and overwrite transformations", "Reimports asset from the same file and overwrites the transformations", EUserInterfaceActionType::RadioButton, FInputChord());
+	UI_COMMAND(ReimportModeNewFile, "Reimport from new file", "Reimports asset from a selected file", EUserInterfaceActionType::RadioButton, FInputChord());
 
 	UI_COMMAND(ToggleFadeIn, "Toggle Fade In", "Toggle adding Fade In Transformation onto the sound wave", EUserInterfaceActionType::Button, FInputChord());
 	UI_COMMAND(FadeInLinear, "Linear Fade In", "Fade In attack curve: 1.0 = linear growth, < 1.0 logorithmic growth, > 1.0 exponential growth", EUserInterfaceActionType::RadioButton, FInputChord());
diff --git a/Engine/Plugins/Editor/WaveformEditor/Source/WaveformEditor/Public/WaveformEditor.h b/Engine/Plugins/Editor/WaveformEditor/Source/WaveformEditor/Public/WaveformEditor.h
--- a/Engine/Plugins/Editor/WaveformEditor/Source/WaveformEditor/Public/WaveformEditor.h
+++ b/Engine/Plugins/Editor/WaveformEditor/Source/WaveformEditor/Public/WaveformEditor.h
@@ -57,7 +57,7 @@ private:
 	bool BindDelegates();
 	bool SetUpAssetReimport();
 	void ExecuteReimport();
-	void ExecuteReimportAndOverwriteTransformations();
+	void ExecuteOverwriteTransformations();
 
 	/**	Sets the wave editor layout */
 	const TSharedRef<FTabManager::FLayout> SetupStandaloneLayout();
@@ -160,7 +160,8 @@ private:
 
 	enum class EWaveEditorReimportMode : uint8
 	{
-		SameFile = 0, 
+		SameFile = 0,
+		SameFileOverwrite,
 		SelectFile,
 		COUNT
 
diff --git a/Engine/Plugins/Editor/WaveformEditor/Source/WaveformTransformationsWidgets/Private/WaveformTransformationMarkersRenderer.cpp b/Engine/Plugins/Editor/WaveformEditor/Source/WaveformTransformationsWidgets/Private/WaveformTransformationMarkersRenderer.cpp
--- a/Engine/Plugins/Editor/WaveformEditor/Source/WaveformTransformationsWidgets/Private/WaveformTransformationMarkersRenderer.cpp
+++ b/Engine/Plugins/Editor/WaveformEditor/Source/WaveformTransformationsWidgets/Private/WaveformTransformationMarkersRenderer.cpp
@@ -31,9 +31,8 @@ int32 FWaveformTransformationMarkerRenderer::DrawMarkerHandles(const FGeometry&
 	check(TransformationWaveInfo.NumChannels > 0);
 	check(TransformationWaveInfo.StartFrameOffset % TransformationWaveInfo.NumChannels == 0);
 	const uint32 StartFrameOffset = TransformationWaveInfo.StartFrameOffset / TransformationWaveInfo.NumChannels;
-	const float WindowWidth = AllottedGeometry.Size.X;
 	const float MarkerHeight = AllottedGeometry.Size.Y * InteractionRatioYDelta;
-	const float MarkerWidth = InteractionPixelXDelta * 0.8;
+	constexpr float MarkerWidth = InteractionPixelXDelta * 0.8;
 	FSlateRoundedBoxBrush RoundedBoxBrush(FLinearColor::Red, 2);
 	FSlateFontInfo Font = FAppStyle::GetFontStyle("Regular");
 	Font.Size = LabelFontSize;
@@ -66,7 +65,6 @@ int32 FWaveformTransformationMarkerRenderer::DrawMarkerHandles(const FGeometry&
 
 		const float MarkerCenter = LoopRegionPosition;
 		const float	HandleStart = MarkerCenter - MarkerWidth / 2;
-		const float	HandleEnd = HandleStart + MarkerWidth;
 		const float LoopRegionPixelLength = LoopRegion.FrameLength * PixelsPerFrame;
 		FSlateBrush Brush;
 		FPaintGeometry HandleGeometry = AllottedGeometry.ToPaintGeometry(FVector2f(MarkerWidth, MarkerHeight), FSlateLayoutTransform(FVector2f(HandleStart, 0)));
@@ -135,7 +133,6 @@ int32 FWaveformTransformationMarkerRenderer::DrawMarkerHandles(const FGeometry&
 		);
 
 		const float	HandleStart = MarkerPosition - MarkerWidth / 2;
-		const float	HandleEnd = HandleStart + MarkerWidth;
 		FPaintGeometry HandleGeometry = AllottedGeometry.ToPaintGeometry(FVector2f(MarkerWidth, MarkerHeight), FSlateLayoutTransform(FVector2f(HandleStart, 0)));
 
 		FSlateDrawElement::MakeBox(
@@ -177,13 +174,9 @@ void FWaveformTransformationMarkerRenderer::Tick(const FGeometry& AllottedGeomet
 	PixelsPerFrame = FMath::Max(AllottedGeometry.GetLocalSize().X / NumFrames, 0.0);
 	ensure(PixelsPerFrame >= 0);
 
+	MarkersArray = CastChecked<UWaveCueArray>(GetPropertyValue<UObject*>(MarkersHandle));
+	
 	TMap<uint32, FSoundWaveCuePoint> CuePointMap;
-
-	MarkersArray = Cast<UWaveCueArray>(GetPropertyValue<UObject*>(MarkersHandle));
-	if (MarkersArray == nullptr)
-	{
-		return;
-	}
 	CuePointMap.Reserve(MarkersArray->WaveCues.Num());
 
 	bool IsMarkerInInteractionRangeValid = false;
@@ -241,7 +234,9 @@ FReply FWaveformTransformationMarkerRenderer::OnMouseButtonDown(SWidget& OwnerWi
 
 		if (MarkerInteractionType != EMarkerInteractionType::None)
 		{
+			MarkersHandle->NotifyPreChange();
 			BeginTransaction(TEXT("PropertyEditor"), FText::Format(NSLOCTEXT("PropertyEditor", "SetNumericPropertyTransaction", "Edit {0}"), FText::FromName(GET_MEMBER_NAME_CHECKED(UWaveformTransformationMarkers, Markers))), nullptr);
+
 			return FReply::Handled().CaptureMouse(OwnerWidget.AsShared()).PreventThrottling();
 		}
 
@@ -254,9 +249,13 @@ FReply FWaveformTransformationMarkerRenderer::OnMouseButtonUp(SWidget& OwnerWidg
 {
 	if (MarkerInteractionType != EMarkerInteractionType::None)
 	{
+		MarkersArray->Modify();
 		SetPropertyValueDependingOnInteractionType(MouseEvent, MyGeometry);
 		EndTransaction();
 		MarkerInteractionType = EMarkerInteractionType::None;
+		
+		MarkersHandle->NotifyPostChange(EPropertyChangeType::ValueSet);
+		
 		return FReply::Handled().ReleaseMouseCapture();
 	}
 
@@ -267,7 +266,7 @@ FReply FWaveformTransformationMarkerRenderer::OnMouseMove(SWidget& OwnerWidget,
 {
 	if (MouseEvent.IsMouseButtonDown(EKeys::LeftMouseButton) && MarkerInteractionType != EMarkerInteractionType::None)
 	{
-		SetPropertyValueDependingOnInteractionType(MouseEvent, MyGeometry, EPropertyValueSetFlags::InteractiveChange | EPropertyValueSetFlags::NotTransactable);
+		SetPropertyValueDependingOnInteractionType(MouseEvent, MyGeometry);
 
 		return FReply::Handled().CaptureMouse(OwnerWidget.AsShared());
 	}
@@ -282,6 +281,18 @@ FCursorReply FWaveformTransformationMarkerRenderer::OnCursorQuery(const FGeometr
 		return FCursorReply::Cursor(EMouseCursor::GrabHandClosed);
 	}
 
+	const FVector2D LocalCursorPosition = GetLocalCursorPosition(CursorEvent, MyGeometry);
+
+	// Set grab hand on mouse hover
+	for (const FSoundWaveCuePoint& CuePoint : MarkersArray->WaveCues)
+	{
+		if (IsPositionInInteractionRange(CuePoint.FramePosition, LocalCursorPosition, MyGeometry) ||
+			(CuePoint.IsLoopRegion() && IsPositionInInteractionRange(CuePoint.FramePosition + CuePoint.FrameLength, LocalCursorPosition, MyGeometry)))
+		{
+			return FCursorReply::Cursor(EMouseCursor::GrabHand);
+		}
+	}
+
 	return FCursorReply::Unhandled();
 }
 
@@ -334,6 +345,7 @@ bool FWaveformTransformationMarkerRenderer::IsPositionInInteractionRange(const i
 	{
 		return true;
 	}
+	
 	return false;
 }
 
@@ -341,11 +353,13 @@ bool FWaveformTransformationMarkerRenderer::IsPositionInInteractionRange(const i
 bool FWaveformTransformationMarkerRenderer::SetMarkerInInteractionRange(const FVector2D& InLocalCursorPosition, const FGeometry& WidgetGeometry)
 {
 	InteractionRanges.Reset(MarkersArray->WaveCues.Num());
+	
 	for (FSoundWaveCuePoint& CuePoint : MarkersArray->WaveCues)
 	{
 		if (IsPositionInInteractionRange(CuePoint.FramePosition, InLocalCursorPosition, WidgetGeometry))
 		{
 			MarkerInInteractionRange.SetMarkerInInteractionRange(&CuePoint);
+			
 			return true;
 		}
 
@@ -359,16 +373,19 @@ bool FWaveformTransformationMarkerRenderer::SetMarkerInInteractionRange(const FV
 		if (IsPositionInInteractionRange(LoopRegionEndHandlePosition, InLocalCursorPosition, WidgetGeometry))
 		{
 			MarkerInInteractionRange.SetMarkerInInteractionRange(&CuePoint, true);
+			
 			return true;
 		}
 	}
 	MarkerInInteractionRange.SetMarkerInInteractionRange(nullptr);
+	
 	return false;
 }
 
-FVector2D FWaveformTransformationMarkerRenderer::GetLocalCursorPosition(const FPointerEvent& MouseEvent, const FGeometry& EventGeometry) const
+FVector2D FWaveformTransformationMarkerRenderer::GetLocalCursorPosition(const FPointerEvent& MouseEvent, const FGeometry& EventGeometry)
 {
 	const FVector2D ScreenSpacePosition = MouseEvent.GetScreenSpacePosition();
+	
 	return  EventGeometry.AbsoluteToLocal(ScreenSpacePosition);
 }
 
@@ -382,7 +399,7 @@ FWaveformTransformationMarkerRenderer::EMarkerInteractionType FWaveformTransform
 	return EMarkerInteractionType::None;
 }
 
-void FWaveformTransformationMarkerRenderer::SetPropertyValueDependingOnInteractionType(const FPointerEvent& MouseEvent, const FGeometry& WidgetGeometry, const EPropertyValueSetFlags::Type Flags)
+void FWaveformTransformationMarkerRenderer::SetPropertyValueDependingOnInteractionType(const FPointerEvent& MouseEvent, const FGeometry& WidgetGeometry)
 {
 	const FVector2D LocalCursorPosition = GetLocalCursorPosition(MouseEvent, WidgetGeometry);
 	const float LocalCursorXClamped = FMath::Clamp(LocalCursorPosition.X, 0.f, WidgetGeometry.GetLocalSize().X);
@@ -407,10 +424,10 @@ void FWaveformTransformationMarkerRenderer::FMarkerInInteractionRange::SetMarker
 {
 	check(MarkerPtr);
 
-	const int32 MinLoopSize = 10;
-
 	if (IsLoopRegionEndMarker)
 	{
+		constexpr int32 MinLoopSize = 10;
+		
 		//FMath::Max prevents the user from collapsing the loop
 		MarkerPtr->FrameLength = FMath::Max(InPosition, MarkerPtr->FramePosition + MinLoopSize) - MarkerPtr->FramePosition;
 	}
diff --git a/Engine/Plugins/Editor/WaveformEditor/Source/WaveformTransformationsWidgets/Private/WaveformTransformationMarkersRenderer.h b/Engine/Plugins/Editor/WaveformEditor/Source/WaveformTransformationsWidgets/Private/WaveformTransformationMarkersRenderer.h
--- a/Engine/Plugins/Editor/WaveformEditor/Source/WaveformTransformationsWidgets/Private/WaveformTransformationMarkersRenderer.h
+++ b/Engine/Plugins/Editor/WaveformEditor/Source/WaveformTransformationsWidgets/Private/WaveformTransformationMarkersRenderer.h
@@ -30,7 +30,7 @@ private:
 
 	bool IsPositionInInteractionRange(const int32 InFramePosition, const FVector2D& InLocalCursorPosition, const FGeometry& WidgetGeometry) const;
 	bool SetMarkerInInteractionRange(const FVector2D& InLocalCursorPosition, const FGeometry& WidgetGeometry);
-	FVector2D GetLocalCursorPosition(const FPointerEvent& MouseEvent, const FGeometry& EventGeometry) const;
+	static FVector2D GetLocalCursorPosition(const FPointerEvent& MouseEvent, const FGeometry& EventGeometry);
 
 	TArray<TRange<float>> InteractionRanges;
 
@@ -53,7 +53,7 @@ private:
 	} MarkerInInteractionRange;
 
 	EMarkerInteractionType GetInteractionTypeFromCursorPosition(const FVector2D& InLocalCursorPosition, const FGeometry& WidgetGeometry);
-	void SetPropertyValueDependingOnInteractionType(const FPointerEvent& MouseEvent, const FGeometry& WidgetGeometry, const EPropertyValueSetFlags::Type Flags = EPropertyValueSetFlags::DefaultFlags);
+	void SetPropertyValueDependingOnInteractionType(const FPointerEvent& MouseEvent, const FGeometry& WidgetGeometry);
 
 	TArray<FSoundWaveCuePoint> CuePoints;
 	TArray<FSoundWaveCuePoint> LoopPoints;
