[Waveform Editor] Make waveform editor default asset editor for SoundWaves

diff --git a/Engine/Plugins/Editor/WaveformEditor/Source/WaveformEditor/Private/WaveformEditor.cpp b/Engine/Plugins/Editor/WaveformEditor/Source/WaveformEditor/Private/WaveformEditor.cpp
--- a/Engine/Plugins/Editor/WaveformEditor/Source/WaveformEditor/Private/WaveformEditor.cpp
+++ b/Engine/Plugins/Editor/WaveformEditor/Source/WaveformEditor/Private/WaveformEditor.cpp
@@ -25,6 +25,7 @@
 #include "WaveformEditorToolMenuContext.h"
 #include "WaveformEditorTransformationsSettings.h"
 #include "WaveformEditorWaveWriter.h"
+#include "WaveformTransformationMarkers.h"
 #include "Widgets/Docking/SDockTab.h"
 
 #define LOCTEXT_NAMESPACE "WaveformEditor"
@@ -273,6 +274,42 @@ bool FWaveformEditor::CanFadeOut()
 	return true;
 }
 
+void FWaveformEditor::CreateMarker()
+{
+	EObjectFlags MaskedOuterFlags = SoundWave ? SoundWave->GetMaskedFlags(RF_PropagateToSubObjects) : RF_NoFlags;
+
+	if (SoundWave->HasAnyFlags(RF_ClassDefaultObject | RF_ArchetypeObject))
+	{
+		MaskedOuterFlags |= RF_ArchetypeObject;
+	}
+
+	UWaveformTransformationMarkers* WaveformMarker = nullptr;
+	for (TObjectPtr<UWaveformTransformationBase> Transformation : SoundWave->Transformations)
+	{
+		WaveformMarker = Cast<UWaveformTransformationMarkers>(Transformation);
+		if (WaveformMarker)
+		{
+			break;
+		}
+	}
+
+	if (WaveformMarker == nullptr)
+	{
+		WaveformMarker = NewObject<UWaveformTransformationMarkers>(
+	SoundWave, UWaveformTransformationMarkers::StaticClass(), NAME_None, MaskedOuterFlags);
+		SoundWave->Transformations.Add(WaveformMarker);
+	}
+
+	TransportCoordinator->SetProgressRatio(0.5);
+
+	// TODO add a new cue point
+	// FSoundWaveCuePoint NewCuePoint = FSoundWaveCuePoint();
+	// WaveformMarker->Markers->WaveCues.Append(&NewCuePoint, 1);
+
+	SoundWave->Modify();
+	NotifyPostTransformationChange(EPropertyChangeType::ValueSet);
+}
+
 bool FWaveformEditor::InitializeAudioComponent()
 {
 	if (!ensure(SoundWave))
@@ -641,6 +678,22 @@ bool FWaveformEditor::RegisterToolbar()
 					InSection.AddEntry(ExportOptionsEntry);
 				}
 			}));
+
+		TransformationSection.AddDynamicEntry("CreateMarkerSection", FNewToolMenuSectionDelegate::CreateLambda([this, Commands](FToolMenuSection& InSection)
+	{
+		const UWaveformEditorToolMenuContext* Context = InSection.FindContext<UWaveformEditorToolMenuContext>();
+
+		if (Context && Context->WaveformEditor.IsValid())
+		{
+			FToolMenuEntry CreateMarkerSection = FToolMenuEntry::InitToolBarButton(
+				Commands.CreateMarker,
+				FText(),
+				LOCTEXT("WaveformEditorCreateMarkerButtonTooltip", "Create a marker"),
+				FSlateIcon("FlippedStyleSet", "Sequencer.Tracks.Fade.Flipped")
+			);
+			InSection.AddEntry(CreateMarkerSection);
+		}
+	}));
 	}
 
 	return true;
@@ -852,6 +905,11 @@ bool FWaveformEditor::BindCommands()
 		FCanExecuteAction::CreateLambda([this] {return true; }),
 		FIsActionChecked::CreateLambda([this] {return FadeOutCurve == UWaveformTransformationTrimFade::FadeModeToCurveValueMap[EWaveEditorFadeMode::Sigmoid]; }));
 
+	ToolkitCommands->MapAction(
+	Commands.CreateMarker,
+	FExecuteAction::CreateSP(this, &FWaveformEditor::CreateMarker),
+	FCanExecuteAction::CreateLambda([this] {return true; }));
+
 	return true;
 }
 
@@ -1070,9 +1128,9 @@ const TSharedRef<FTabManager::FLayout> FWaveformEditor::SetupStandaloneLayout()
 				(
 					FTabManager::NewStack()
 					->SetSizeCoefficient(0.2f)
-					->AddTab(TransformationsTabId, ETabState::OpenedTab)
 					->AddTab(PropertiesTabId, ETabState::OpenedTab)
-					->SetForegroundTab(TransformationsTabId)
+					->AddTab(TransformationsTabId, ETabState::OpenedTab)
+					->SetForegroundTab(PropertiesTabId)
 
 				)
 				->Split
diff --git a/Engine/Plugins/Editor/WaveformEditor/Source/WaveformEditor/Private/WaveformEditorInstantiator.cpp b/Engine/Plugins/Editor/WaveformEditor/Source/WaveformEditor/Private/WaveformEditorInstantiator.cpp
--- a/Engine/Plugins/Editor/WaveformEditor/Source/WaveformEditor/Private/WaveformEditorInstantiator.cpp
+++ b/Engine/Plugins/Editor/WaveformEditor/Source/WaveformEditor/Private/WaveformEditorInstantiator.cpp
@@ -3,6 +3,7 @@
 #include "WaveformEditorInstantiator.h"
 
 #include "Algo/AnyOf.h"
+#include "AudioEditorModule.h"
 #include "ContentBrowserMenuContexts.h"
 #include "EditorReimportHandler.h"
 #include "Framework/Application/SlateApplication.h"
@@ -63,6 +64,12 @@ void FWaveformEditorInstantiator::ExtendContentBrowserSelectionMenu()
 		}));
 }
 
+void FWaveformEditorInstantiator::RegisterAsSoundwaveEditor()
+{
+	IAudioEditorModule& AudioEditorModule = FModuleManager::LoadModuleChecked<IAudioEditorModule>( "AudioEditor" );
+	AudioEditorModule.SoundWaveEditorOpen.BindSP(this, &FWaveformEditorInstantiator::CreateWaveformEditor);
+}
+
 void FWaveformEditorInstantiator::ExecuteCreateWaveformEditor(const FToolMenuContext& MenuContext)
 {
 	if (const UContentBrowserAssetContextMenuContext* Context = UContentBrowserAssetContextMenuContext::FindContextWithAssets(MenuContext))
diff --git a/Engine/Plugins/Editor/WaveformEditor/Source/WaveformEditor/Private/WaveformEditorInstantiator.h b/Engine/Plugins/Editor/WaveformEditor/Source/WaveformEditor/Private/WaveformEditorInstantiator.h
--- a/Engine/Plugins/Editor/WaveformEditor/Source/WaveformEditor/Private/WaveformEditorInstantiator.h
+++ b/Engine/Plugins/Editor/WaveformEditor/Source/WaveformEditor/Private/WaveformEditorInstantiator.h
@@ -12,6 +12,9 @@ class FWaveformEditorInstantiator : public IWaveformEditorInstantiator,  public
 {
 public:
 	virtual void ExtendContentBrowserSelectionMenu() override;
+	
+	// Registers using a delegate as the default editor for soundwave assets without main engine referencing waveform plugin in Sound Wave Asset Definition
+	virtual void RegisterAsSoundwaveEditor() override;
 	virtual ~FWaveformEditorInstantiator() = default;
 
 private:
diff --git a/Engine/Plugins/Editor/WaveformEditor/Source/WaveformEditor/Private/WaveformEditorModule.cpp b/Engine/Plugins/Editor/WaveformEditor/Source/WaveformEditor/Private/WaveformEditorModule.cpp
--- a/Engine/Plugins/Editor/WaveformEditor/Source/WaveformEditor/Private/WaveformEditorModule.cpp
+++ b/Engine/Plugins/Editor/WaveformEditor/Source/WaveformEditor/Private/WaveformEditorModule.cpp
@@ -19,6 +19,7 @@ void FWaveformEditorModule::StartupModule()
 
 	WaveformEditorInstantiator = MakeShared<FWaveformEditorInstantiator>();
 	RegisterContentBrowserExtensions(WaveformEditorInstantiator.Get());
+	WaveformEditorInstantiator->RegisterAsSoundwaveEditor();
 }
 
 void FWaveformEditorModule::ShutdownModule()
diff --git a/Engine/Plugins/Editor/WaveformEditor/Source/WaveformEditor/Public/IWaveformEditorInstantiator.h b/Engine/Plugins/Editor/WaveformEditor/Source/WaveformEditor/Public/IWaveformEditorInstantiator.h
--- a/Engine/Plugins/Editor/WaveformEditor/Source/WaveformEditor/Public/IWaveformEditorInstantiator.h
+++ b/Engine/Plugins/Editor/WaveformEditor/Source/WaveformEditor/Public/IWaveformEditorInstantiator.h
@@ -10,6 +10,7 @@ class IWaveformEditorInstantiator
 {
 public:
 	virtual void ExtendContentBrowserSelectionMenu() = 0;
+	virtual void RegisterAsSoundwaveEditor() = 0;
 
 private:
 	virtual void CreateWaveformEditor(TArray<USoundWave*> SoundWavesToEdit) = 0;
diff --git a/Engine/Source/Editor/AudioEditor/Private/AssetTypeActions/AssetDefinition_SoundWave.cpp b/Engine/Source/Editor/AudioEditor/Private/AssetTypeActions/AssetDefinition_SoundWave.cpp
--- a/Engine/Source/Editor/AudioEditor/Private/AssetTypeActions/AssetDefinition_SoundWave.cpp
+++ b/Engine/Source/Editor/AudioEditor/Private/AssetTypeActions/AssetDefinition_SoundWave.cpp
@@ -47,10 +47,20 @@ TConstArrayView<FAssetCategoryPath> UAssetDefinition_SoundWave::GetAssetCategori
 EAssetCommandResult UAssetDefinition_SoundWave::OpenAssets(const FAssetOpenArgs& OpenArgs) const
 {
 	IAudioEditorModule& AudioEditorModule = FModuleManager::LoadModuleChecked<IAudioEditorModule>( "AudioEditor" );
-	for (USoundWave* SoundWave : OpenArgs.LoadObjects<USoundWave>())
+
+	// Use a custom editor for soundwaves if one is registered
+	if (AudioEditorModule.SoundWaveEditorOpen.IsBound())
+	{
+		AudioEditorModule.SoundWaveEditorOpen.Execute(OpenArgs.LoadObjects<USoundWave>());
+	}
+	else // Otherwise open a default editor
 	{
-		FSimpleAssetEditor::CreateEditor(EToolkitMode::Standalone, OpenArgs.ToolkitHost, SoundWave);
+		for (USoundWave* SoundWave : OpenArgs.LoadObjects<USoundWave>())
+		{
+			FSimpleAssetEditor::CreateEditor(EToolkitMode::Standalone, OpenArgs.ToolkitHost, SoundWave);
+		}
 	}
+
 	return EAssetCommandResult::Handled;
 }
 
diff --git a/Engine/Source/Editor/AudioEditor/Public/AudioEditorModule.h b/Engine/Source/Editor/AudioEditor/Public/AudioEditorModule.h
--- a/Engine/Source/Editor/AudioEditor/Public/AudioEditorModule.h
+++ b/Engine/Source/Editor/AudioEditor/Public/AudioEditorModule.h
@@ -94,4 +94,9 @@ public:
 
 	// Sets if the Audio editor is in "restricted mode" (certain asset actions may be disabled)
 	virtual void SetRestrictedMode(bool bInRestricted) = 0;
+
+	// Allows systems to register to open soundwave files with custom editor
+	DECLARE_DELEGATE_OneParam(SoundWaveEditorOpenDelegate, TArray<USoundWave*>);
+
+	SoundWaveEditorOpenDelegate SoundWaveEditorOpen;
 };
