[Backout] - CL40832676
#rnx
Original CL Desc
-----------------------------------------------------------------
[Waveform Editor] Make waveform editor default asset editor for SoundWaves

diff --git a/Engine/Plugins/Editor/WaveformEditor/Source/WaveformEditor/Private/WaveformEditor.cpp b/Engine/Plugins/Editor/WaveformEditor/Source/WaveformEditor/Private/WaveformEditor.cpp
--- a/Engine/Plugins/Editor/WaveformEditor/Source/WaveformEditor/Private/WaveformEditor.cpp
+++ b/Engine/Plugins/Editor/WaveformEditor/Source/WaveformEditor/Private/WaveformEditor.cpp
@@ -25,7 +25,6 @@
 #include "WaveformEditorToolMenuContext.h"
 #include "WaveformEditorTransformationsSettings.h"
 #include "WaveformEditorWaveWriter.h"
-#include "WaveformTransformationMarkers.h"
 #include "Widgets/Docking/SDockTab.h"
 
 #define LOCTEXT_NAMESPACE "WaveformEditor"
@@ -274,42 +273,6 @@ bool FWaveformEditor::CanFadeOut()
 	return true;
 }
 
-void FWaveformEditor::CreateMarker()
-{
-	EObjectFlags MaskedOuterFlags = SoundWave ? SoundWave->GetMaskedFlags(RF_PropagateToSubObjects) : RF_NoFlags;
-
-	if (SoundWave->HasAnyFlags(RF_ClassDefaultObject | RF_ArchetypeObject))
-	{
-		MaskedOuterFlags |= RF_ArchetypeObject;
-	}
-
-	UWaveformTransformationMarkers* WaveformMarker = nullptr;
-	for (TObjectPtr<UWaveformTransformationBase> Transformation : SoundWave->Transformations)
-	{
-		WaveformMarker = Cast<UWaveformTransformationMarkers>(Transformation);
-		if (WaveformMarker)
-		{
-			break;
-		}
-	}
-
-	if (WaveformMarker == nullptr)
-	{
-		WaveformMarker = NewObject<UWaveformTransformationMarkers>(
-	SoundWave, UWaveformTransformationMarkers::StaticClass(), NAME_None, MaskedOuterFlags);
-		SoundWave->Transformations.Add(WaveformMarker);
-	}
-
-	TransportCoordinator->SetProgressRatio(0.5);
-
-	// TODO add a new cue point
-	// FSoundWaveCuePoint NewCuePoint = FSoundWaveCuePoint();
-	// WaveformMarker->Markers->WaveCues.Append(&NewCuePoint, 1);
-
-	SoundWave->Modify();
-	NotifyPostTransformationChange(EPropertyChangeType::ValueSet);
-}
-
 bool FWaveformEditor::InitializeAudioComponent()
 {
 	if (!ensure(SoundWave))
@@ -678,22 +641,6 @@ bool FWaveformEditor::RegisterToolbar()
 					InSection.AddEntry(ExportOptionsEntry);
 				}
 			}));
-
-		TransformationSection.AddDynamicEntry("CreateMarkerSection", FNewToolMenuSectionDelegate::CreateLambda([this, Commands](FToolMenuSection& InSection)
-	{
-		const UWaveformEditorToolMenuContext* Context = InSection.FindContext<UWaveformEditorToolMenuContext>();
-
-		if (Context && Context->WaveformEditor.IsValid())
-		{
-			FToolMenuEntry CreateMarkerSection = FToolMenuEntry::InitToolBarButton(
-				Commands.CreateMarker,
-				FText(),
-				LOCTEXT("WaveformEditorCreateMarkerButtonTooltip", "Create a marker"),
-				FSlateIcon("FlippedStyleSet", "Sequencer.Tracks.Fade.Flipped")
-			);
-			InSection.AddEntry(CreateMarkerSection);
-		}
-	}));
 	}
 
 	return true;
@@ -905,11 +852,6 @@ bool FWaveformEditor::BindCommands()
 		FCanExecuteAction::CreateLambda([this] {return true; }),
 		FIsActionChecked::CreateLambda([this] {return FadeOutCurve == UWaveformTransformationTrimFade::FadeModeToCurveValueMap[EWaveEditorFadeMode::Sigmoid]; }));
 
-	ToolkitCommands->MapAction(
-	Commands.CreateMarker,
-	FExecuteAction::CreateSP(this, &FWaveformEditor::CreateMarker),
-	FCanExecuteAction::CreateLambda([this] {return true; }));
-
 	return true;
 }
 
@@ -1128,9 +1070,9 @@ const TSharedRef<FTabManager::FLayout> FWaveformEditor::SetupStandaloneLayout()
 				(
 					FTabManager::NewStack()
 					->SetSizeCoefficient(0.2f)
-					->AddTab(PropertiesTabId, ETabState::OpenedTab)
 					->AddTab(TransformationsTabId, ETabState::OpenedTab)
-					->SetForegroundTab(PropertiesTabId)
+					->AddTab(PropertiesTabId, ETabState::OpenedTab)
+					->SetForegroundTab(TransformationsTabId)
 
 				)
 				->Split
diff --git a/Engine/Plugins/Editor/WaveformEditor/Source/WaveformEditor/Private/WaveformEditorInstantiator.cpp b/Engine/Plugins/Editor/WaveformEditor/Source/WaveformEditor/Private/WaveformEditorInstantiator.cpp
--- a/Engine/Plugins/Editor/WaveformEditor/Source/WaveformEditor/Private/WaveformEditorInstantiator.cpp
+++ b/Engine/Plugins/Editor/WaveformEditor/Source/WaveformEditor/Private/WaveformEditorInstantiator.cpp
@@ -3,7 +3,6 @@
 #include "WaveformEditorInstantiator.h"
 
 #include "Algo/AnyOf.h"
-#include "AudioEditorModule.h"
 #include "ContentBrowserMenuContexts.h"
 #include "EditorReimportHandler.h"
 #include "Framework/Application/SlateApplication.h"
@@ -64,12 +63,6 @@ void FWaveformEditorInstantiator::ExtendContentBrowserSelectionMenu()
 		}));
 }
 
-void FWaveformEditorInstantiator::RegisterAsSoundwaveEditor()
-{
-	IAudioEditorModule& AudioEditorModule = FModuleManager::LoadModuleChecked<IAudioEditorModule>( "AudioEditor" );
-	AudioEditorModule.SoundWaveEditorOpen.BindSP(this, &FWaveformEditorInstantiator::CreateWaveformEditor);
-}
-
 void FWaveformEditorInstantiator::ExecuteCreateWaveformEditor(const FToolMenuContext& MenuContext)
 {
 	if (const UContentBrowserAssetContextMenuContext* Context = UContentBrowserAssetContextMenuContext::FindContextWithAssets(MenuContext))
diff --git a/Engine/Plugins/Editor/WaveformEditor/Source/WaveformEditor/Private/WaveformEditorInstantiator.h b/Engine/Plugins/Editor/WaveformEditor/Source/WaveformEditor/Private/WaveformEditorInstantiator.h
--- a/Engine/Plugins/Editor/WaveformEditor/Source/WaveformEditor/Private/WaveformEditorInstantiator.h
+++ b/Engine/Plugins/Editor/WaveformEditor/Source/WaveformEditor/Private/WaveformEditorInstantiator.h
@@ -12,9 +12,6 @@ class FWaveformEditorInstantiator : public IWaveformEditorInstantiator,  public
 {
 public:
 	virtual void ExtendContentBrowserSelectionMenu() override;
-	
-	// Registers using a delegate as the default editor for soundwave assets without main engine referencing waveform plugin in Sound Wave Asset Definition
-	virtual void RegisterAsSoundwaveEditor() override;
 	virtual ~FWaveformEditorInstantiator() = default;
 
 private:
diff --git a/Engine/Plugins/Editor/WaveformEditor/Source/WaveformEditor/Private/WaveformEditorModule.cpp b/Engine/Plugins/Editor/WaveformEditor/Source/WaveformEditor/Private/WaveformEditorModule.cpp
--- a/Engine/Plugins/Editor/WaveformEditor/Source/WaveformEditor/Private/WaveformEditorModule.cpp
+++ b/Engine/Plugins/Editor/WaveformEditor/Source/WaveformEditor/Private/WaveformEditorModule.cpp
@@ -19,7 +19,6 @@ void FWaveformEditorModule::StartupModule()
 
 	WaveformEditorInstantiator = MakeShared<FWaveformEditorInstantiator>();
 	RegisterContentBrowserExtensions(WaveformEditorInstantiator.Get());
-	WaveformEditorInstantiator->RegisterAsSoundwaveEditor();
 }
 
 void FWaveformEditorModule::ShutdownModule()
diff --git a/Engine/Plugins/Editor/WaveformEditor/Source/WaveformEditor/Public/IWaveformEditorInstantiator.h b/Engine/Plugins/Editor/WaveformEditor/Source/WaveformEditor/Public/IWaveformEditorInstantiator.h
--- a/Engine/Plugins/Editor/WaveformEditor/Source/WaveformEditor/Public/IWaveformEditorInstantiator.h
+++ b/Engine/Plugins/Editor/WaveformEditor/Source/WaveformEditor/Public/IWaveformEditorInstantiator.h
@@ -10,7 +10,6 @@ class IWaveformEditorInstantiator
 {
 public:
 	virtual void ExtendContentBrowserSelectionMenu() = 0;
-	virtual void RegisterAsSoundwaveEditor() = 0;
 
 private:
 	virtual void CreateWaveformEditor(TArray<USoundWave*> SoundWavesToEdit) = 0;
diff --git a/Engine/Source/Editor/AudioEditor/Private/AssetTypeActions/AssetDefinition_SoundWave.cpp b/Engine/Source/Editor/AudioEditor/Private/AssetTypeActions/AssetDefinition_SoundWave.cpp
--- a/Engine/Source/Editor/AudioEditor/Private/AssetTypeActions/AssetDefinition_SoundWave.cpp
+++ b/Engine/Source/Editor/AudioEditor/Private/AssetTypeActions/AssetDefinition_SoundWave.cpp
@@ -47,20 +47,10 @@ TConstArrayView<FAssetCategoryPath> UAssetDefinition_SoundWave::GetAssetCategori
 EAssetCommandResult UAssetDefinition_SoundWave::OpenAssets(const FAssetOpenArgs& OpenArgs) const
 {
 	IAudioEditorModule& AudioEditorModule = FModuleManager::LoadModuleChecked<IAudioEditorModule>( "AudioEditor" );
-
-	// Use a custom editor for soundwaves if one is registered
-	if (AudioEditorModule.SoundWaveEditorOpen.IsBound())
-	{
-		AudioEditorModule.SoundWaveEditorOpen.Execute(OpenArgs.LoadObjects<USoundWave>());
-	}
-	else // Otherwise open a default editor
+	for (USoundWave* SoundWave : OpenArgs.LoadObjects<USoundWave>())
 	{
-		for (USoundWave* SoundWave : OpenArgs.LoadObjects<USoundWave>())
-		{
-			FSimpleAssetEditor::CreateEditor(EToolkitMode::Standalone, OpenArgs.ToolkitHost, SoundWave);
-		}
+		FSimpleAssetEditor::CreateEditor(EToolkitMode::Standalone, OpenArgs.ToolkitHost, SoundWave);
 	}
-
 	return EAssetCommandResult::Handled;
 }
 
diff --git a/Engine/Source/Editor/AudioEditor/Public/AudioEditorModule.h b/Engine/Source/Editor/AudioEditor/Public/AudioEditorModule.h
--- a/Engine/Source/Editor/AudioEditor/Public/AudioEditorModule.h
+++ b/Engine/Source/Editor/AudioEditor/Public/AudioEditorModule.h
@@ -94,9 +94,4 @@ public:
 
 	// Sets if the Audio editor is in "restricted mode" (certain asset actions may be disabled)
 	virtual void SetRestrictedMode(bool bInRestricted) = 0;
-
-	// Allows systems to register to open soundwave files with custom editor
-	DECLARE_DELEGATE_OneParam(SoundWaveEditorOpenDelegate, TArray<USoundWave*>);
-
-	SoundWaveEditorOpenDelegate SoundWaveEditorOpen;
 };
