[Waveform Editor] - Ensure when clicking Reset on a marker

diff --git a/Engine/Plugins/Editor/WaveformEditor/Source/WaveformTransformations/Private/WaveformTransformationMarkers.cpp b/Engine/Plugins/Editor/WaveformEditor/Source/WaveformTransformations/Private/WaveformTransformationMarkers.cpp
--- a/Engine/Plugins/Editor/WaveformEditor/Source/WaveformTransformations/Private/WaveformTransformationMarkers.cpp
+++ b/Engine/Plugins/Editor/WaveformEditor/Source/WaveformTransformations/Private/WaveformTransformationMarkers.cpp
@@ -71,7 +71,8 @@ void UWaveCueArray::PostEditChangeProperty(FPropertyChangedEvent& PropertyChange
 
 	const FName PropertyName = PropertyChangedEvent.GetPropertyName();
 	if (PropertyName == GET_MEMBER_NAME_CHECKED(UWaveCueArray, WaveCues) &&
-		PropertyChangedEvent.ChangeType == EPropertyChangeType::ArrayAdd)
+		(PropertyChangedEvent.ChangeType & (EPropertyChangeType::ArrayAdd | EPropertyChangeType::ResetToDefault | 
+			EPropertyChangeType::Unspecified | EPropertyChangeType::Redirected)))
 	{
 		check(WaveCues.Num() > 0);
 
