[Waveform Editor] - Fade Curve Options on toggle don't have any effect

diff --git a/Engine/Plugins/Editor/WaveformEditor/Source/WaveformEditor/Private/WaveformEditor.cpp b/Engine/Plugins/Editor/WaveformEditor/Source/WaveformEditor/Private/WaveformEditor.cpp
--- a/Engine/Plugins/Editor/WaveformEditor/Source/WaveformEditor/Private/WaveformEditor.cpp
+++ b/Engine/Plugins/Editor/WaveformEditor/Source/WaveformEditor/Private/WaveformEditor.cpp
@@ -597,7 +597,7 @@ bool FWaveformEditor::RegisterToolbar()
 						FToolUIActionChoice(FUIAction()),
 						FNewToolMenuChoice(FOnGetContent::CreateSP(LockedObserver.Get(), &FWaveformEditor::GenerateFadeInOptionsMenu)),
 						LOCTEXT("FadeInOptions_Label", "Fade In Options"),
-						LOCTEXT("FadeInOptions_ToolTip", "Fade In options for this waveform"),
+						LOCTEXT("FadeInOptions_ToolTip", "Fade In options for this waveform. Applied when Toggle Fade In is pressed"),
 						FSlateIcon(FAppStyle::GetAppStyleSetName(), "LevelEditor.ExportAll"),
 						true
 					);
@@ -633,7 +633,7 @@ bool FWaveformEditor::RegisterToolbar()
 						FToolUIActionChoice(FUIAction()),
 						FNewToolMenuChoice(FOnGetContent::CreateSP(LockedObserver.Get(), &FWaveformEditor::GenerateFadeOutOptionsMenu)),
 						LOCTEXT("FadeOutOptions_Label", "Fade Out Options"),
-						LOCTEXT("FadeOutOptions_ToolTip", "Fade Out options for this waveform"),
+						LOCTEXT("FadeOutOptions_ToolTip", "Fade Out options for this waveform. Applied when Toggle Fade Out is pressed"),
 						FSlateIcon(FAppStyle::GetAppStyleSetName(), "LevelEditor.ExportAll"),
 						true
 					);
diff --git a/Engine/Plugins/Editor/WaveformEditor/Source/WaveformEditor/Private/WaveformEditorCommands.cpp b/Engine/Plugins/Editor/WaveformEditor/Source/WaveformEditor/Private/WaveformEditorCommands.cpp
--- a/Engine/Plugins/Editor/WaveformEditor/Source/WaveformEditor/Private/WaveformEditorCommands.cpp
+++ b/Engine/Plugins/Editor/WaveformEditor/Source/WaveformEditor/Private/WaveformEditorCommands.cpp
@@ -33,16 +33,16 @@ void FWaveformEditorCommands::RegisterCommands()
 	UI_COMMAND(ReimportModeNewFile, "Reimport from new file", "Reimports asset from a selected file", EUserInterfaceActionType::RadioButton, FInputChord());
 
 	UI_COMMAND(ToggleFadeIn, "Toggle Fade In", "Toggle adding Fade In Transformation onto the sound wave", EUserInterfaceActionType::Button, FInputChord());
-	UI_COMMAND(FadeInLinear, "Linear Fade In", "Fade In attack curve: 1.0 = linear growth, < 1.0 logorithmic growth, > 1.0 exponential growth", EUserInterfaceActionType::RadioButton, FInputChord());
-	UI_COMMAND(FadeInExponential, "Exponential Fade In", "Fade In attack curve: 1.0 = linear growth, < 1.0 logarithmic growth, > 1.0 exponential growth", EUserInterfaceActionType::RadioButton, FInputChord());
-	UI_COMMAND(FadeInLogarithmic, "Logarithmic Fade In", "Fade In attack curve: 1.0 = linear growth, < 1.0 logarithmic growth, > 1.0 exponential growth", EUserInterfaceActionType::RadioButton, FInputChord());
-	UI_COMMAND(FadeInSigmoid, "Sigmoid Fade In", "Fade In S-Shaped curve", EUserInterfaceActionType::RadioButton, FInputChord());
+	UI_COMMAND(FadeInLinear, "Linear Fade In", "Fade In Curve type applied when Toggle Fade In pressed", EUserInterfaceActionType::RadioButton, FInputChord());
+	UI_COMMAND(FadeInExponential, "Exponential Fade In", "Fade In Curve type applied when Toggle Fade In pressed", EUserInterfaceActionType::RadioButton, FInputChord());
+	UI_COMMAND(FadeInLogarithmic, "Logarithmic Fade In", "Fade In Curve type applied when Toggle Fade In pressed", EUserInterfaceActionType::RadioButton, FInputChord());
+	UI_COMMAND(FadeInSigmoid, "Sigmoid Fade In", "Fade In S-Shaped Curve type applied when Toggle Fade In pressed", EUserInterfaceActionType::RadioButton, FInputChord());
 	
 	UI_COMMAND(ToggleFadeOut, "Toggle Fade Out", "Toggle adding Fade Out Transformation onto the sound wave", EUserInterfaceActionType::Button, FInputChord());
-	UI_COMMAND(FadeOutLinear, "Linear Fade Out", "Fade Out attack curve: 1.0 = linear growth, < 1.0 logorithmic growth, > 1.0 exponential growth", EUserInterfaceActionType::RadioButton, FInputChord());
-	UI_COMMAND(FadeOutExponential, "Exponential Fade Out", "Fade Out attack curve: 1.0 = linear growth, < 1.0 logarithmic growth, > 1.0 exponential growth", EUserInterfaceActionType::RadioButton, FInputChord());
-	UI_COMMAND(FadeOutLogarithmic, "Logarithmic Fade Out", "Fade Out attack curve: 1.0 = linear growth, < 1.0 logarithmic growth, > 1.0 exponential growth", EUserInterfaceActionType::RadioButton, FInputChord());
-	UI_COMMAND(FadeOutSigmoid, "Sigmoid Fade Out", "Fade Out S-Shaped curve", EUserInterfaceActionType::RadioButton, FInputChord());
+	UI_COMMAND(FadeOutLinear, "Linear Fade Out", "Fade Out Curve type applied when Toggle Fade Out pressed", EUserInterfaceActionType::RadioButton, FInputChord());
+	UI_COMMAND(FadeOutExponential, "Exponential Fade Out", "Fade Out Curve type applied when Toggle Fade Out pressed", EUserInterfaceActionType::RadioButton, FInputChord());
+	UI_COMMAND(FadeOutLogarithmic, "Logarithmic Fade Out", "Fade Out Curve type applied when Toggle Fade Out pressed", EUserInterfaceActionType::RadioButton, FInputChord());
+	UI_COMMAND(FadeOutSigmoid, "Sigmoid Fade Out", "Fade Out S-Shaped Curve type applied when Toggle Fade Out pressed", EUserInterfaceActionType::RadioButton, FInputChord());
 
 }
 	
