[Waveform Editor] - Playhead can get stuck after trimming a sound and playing it back at different times.
https://jira.it.epicgames.com/browse/UE-188517
Also, ensure playhead is initialized properly with respect to a persisted trimfade.

diff --git a/Engine/Plugins/Editor/WaveformEditor/Source/WaveformEditor/Private/WaveformEditor.cpp b/Engine/Plugins/Editor/WaveformEditor/Source/WaveformEditor/Private/WaveformEditor.cpp
--- a/Engine/Plugins/Editor/WaveformEditor/Source/WaveformEditor/Private/WaveformEditor.cpp
+++ b/Engine/Plugins/Editor/WaveformEditor/Source/WaveformEditor/Private/WaveformEditor.cpp
@@ -43,7 +43,15 @@ bool FWaveformEditor::Init(const EToolkitMode::Type Mode, const TSharedPtr<ITool
 
 	SoundWave = SoundWaveToEdit;
 
-	AddDefaultTransformations();
+	// Initialize TransformationChainConfig from persisted transformations or if there are none falling back to default transformations.
+	if (SoundWave->Transformations.Num() == 0)
+	{
+		AddDefaultTransformations();
+	}
+	else
+	{
+		TransformationChainConfig = SoundWave->UpdateTransformations();
+	}
 
 	bool bIsInitialized = true;
 	
@@ -81,6 +89,17 @@ bool FWaveformEditor::Init(const EToolkitMode::Type Mode, const TSharedPtr<ITool
 			SoundWaveToEdit,
 			bToolbarFocusable,
 			bUseSmallIcons);
+
+			// Initialize waveform view to with playhead at start of active transformations.
+			check(WaveformView.DataProvider != nullptr);
+			WaveformView.DataProvider->GenerateLayersChain();
+			WaveformView.DataProvider->UpdateRenderElements();
+			PlaybackTimeBeforeTransformInteraction = 0;
+			StartTimeBeforeTransformInteraction = TransformationChainConfig.StartTime; // StartTime=0, EndTime=-1 if no transformations
+			check(TransportController != nullptr);
+			TransportController->CacheStartTime(PlaybackTimeBeforeTransformInteraction);
+			check(TransportCoordinator != nullptr);
+			TransportCoordinator->SetProgressRatio(0.0f);
 	}
 
 	return bIsInitialized;
@@ -906,6 +925,7 @@ void FWaveformEditor::OnAssetReimport(UObject* ReimportedObject, bool bSuccessfu
 	}
 }
 
+
 void FWaveformEditor::NotifyPostChange(const FPropertyChangedEvent& PropertyChangedEvent, class FEditPropertyChain* PropertyThatChanged)
 {	
 	TDoubleLinkedList<FProperty*>::TDoubleLinkedListNode* PropertyNode = PropertyThatChanged->GetActiveMemberNode();
@@ -941,13 +961,32 @@ void FWaveformEditor::NotifyPostChange(const FPropertyChangedEvent& PropertyChan
 		}
 	}
 
+	const bool bUpdateTransformationChain = PropertyChangedEvent.GetPropertyName() == TEXT("Transformations");
+	if (bUpdateTransformationChain)
+	{
+		WaveformView.DataProvider->GenerateLayersChain();
+	}
+	WaveformView.DataProvider->UpdateRenderElements();
+	TransformationChainConfig = SoundWave->GetTransformationChainConfig();
+
 	if (PropertyChangedEvent.ChangeType == EPropertyChangeType::ValueSet)
 	{
 		if (!TransportController->IsPlaying())
 		{
+			// StartTime=0, EndTime=-1 if no transformations
+			const float TransformationDuration = TransformationChainConfig.EndTime - TransformationChainConfig.StartTime;
+
+			// Clamp Playback/Start times to start if out of range of TransformationChainConfig.
+			// If no active+initialized transformation, TransformationDuration will be <= 0.0f, ignore in this case.
+			if ( (PlaybackTimeBeforeTransformInteraction + StartTimeBeforeTransformInteraction - TransformationChainConfig.StartTime <= 0)
+			  || (TransformationDuration > 0.0f && PlaybackTimeBeforeTransformInteraction + StartTimeBeforeTransformInteraction - TransformationChainConfig.StartTime >= TransformationDuration) )
+			{
+				PlaybackTimeBeforeTransformInteraction = 0;
+				StartTimeBeforeTransformInteraction = TransformationChainConfig.StartTime;
+			}
 			const float StartTimeDifference = StartTimeBeforeTransformInteraction - TransformationChainConfig.StartTime;
 			const float AdjustedPlaybackTime = PlaybackTimeBeforeTransformInteraction + StartTimeDifference;
-
+			
 			switch (TransformInteractionPlayState)
 			{
 			case EAudioComponentPlayState::Playing:
@@ -955,8 +994,16 @@ void FWaveformEditor::NotifyPostChange(const FPropertyChangedEvent& PropertyChan
 				TransformInteractionPlayState = EAudioComponentPlayState::Stopped;
 				break;
 			case EAudioComponentPlayState::Paused:
-				TransportController->CacheStartTime(AdjustedPlaybackTime);
 				TransformInteractionPlayState = EAudioComponentPlayState::Stopped;
+			case EAudioComponentPlayState::Stopped:
+				// Set TransportController and displayed TransportCoordinator to value clamped within TransformationChainConfig.
+				TransportController->CacheStartTime(AdjustedPlaybackTime);
+				if (!TransportCoordinator->IsScrubbing())
+				{
+					const float CachedAudioStartTimeAsPercentage = TransportController->GetCachedAudioStartTimeAsPercentage();
+					TransportCoordinator->SetProgressRatio(CachedAudioStartTimeAsPercentage);
+				}
+				break;
 			default:
 				break;
 			}
@@ -964,16 +1011,6 @@ void FWaveformEditor::NotifyPostChange(const FPropertyChangedEvent& PropertyChan
 
 		bIsInteractingWithTransformations = false;
 	}
-
-	const bool bUpdateTransformationChain = PropertyChangedEvent.GetPropertyName() == TEXT("Transformations");
-
-	if (bUpdateTransformationChain)
-	{
-		WaveformView.DataProvider->GenerateLayersChain();
-	}
-
-	WaveformView.DataProvider->UpdateRenderElements();
-	TransformationChainConfig = SoundWave->GetTransformationChainConfig();
 }
 
 
@@ -1226,7 +1263,17 @@ void FWaveformEditor::HandlePlayheadScrub(const float InTargetPlayBackRatio, con
 	}
 	else
 	{
-		const float NewTime = InTargetPlayBackRatio * SoundWave->Duration;
+		float ActiveDuration = TransformationChainConfig.EndTime - TransformationChainConfig.StartTime;
+		if (ActiveDuration <= 0.0f)
+		{
+			// If no active+initialized transformation, ActiveDuration will be <= 0.0f, so fallback to sound duration:
+			ActiveDuration = SoundWave->Duration; // not GetDuration() to get the raw duration and not INDEFINITELY_LOOPING_DURATION if looping
+		}
+		const float ClampedTargetPlayBackRatio = FMath::Max(InTargetPlayBackRatio, 0.0f);
+		const float NewTime = FGenericPlatformMath::Fmod(ClampedTargetPlayBackRatio, 1.0f) * ActiveDuration;
+
+		PlaybackTimeBeforeTransformInteraction = NewTime;
+		StartTimeBeforeTransformInteraction = TransformationChainConfig.StartTime;
 
 		if (TransportController->IsPlaying())
 		{
@@ -1243,7 +1290,6 @@ void FWaveformEditor::HandlePlayheadScrub(const float InTargetPlayBackRatio, con
 		{
 			TransportController->CacheStartTime(NewTime);
 		}
-
 	}
 }
 
diff --git a/Engine/Plugins/Editor/WaveformEditor/Source/WaveformEditor/Private/WaveformEditorTransportController.cpp b/Engine/Plugins/Editor/WaveformEditor/Source/WaveformEditor/Private/WaveformEditorTransportController.cpp
--- a/Engine/Plugins/Editor/WaveformEditor/Source/WaveformEditor/Private/WaveformEditorTransportController.cpp
+++ b/Engine/Plugins/Editor/WaveformEditor/Source/WaveformEditor/Private/WaveformEditorTransportController.cpp
@@ -1,7 +1,11 @@
 // Copyright Epic Games, Inc. All Rights Reserved.
 
 #include "WaveformEditorTransportController.h"
+
+#include "GenericPlatform/GenericPlatformMath.h"
+#include "IWaveformTransformationRenderer.h"
 #include "Sound/SoundBase.h"
+#include "Sound/SoundWave.h"
 
 FWaveformEditorTransportController::FWaveformEditorTransportController(UAudioComponent* InAudioComponent)
 	: AudioComponent(InAudioComponent)
@@ -133,18 +137,30 @@ float FWaveformEditorTransportController::GetCachedAudioStartTimeAsPercentage()
 {
 	check(AudioComponent != nullptr);
 
+	float CachedAudioStartTimeAsPercentage = 0.0f;
+	
 	if (AudioComponent->Sound != nullptr)
 	{
-		const float SoundDuration = AudioComponent->Sound->Duration; // not GetDuration() to get the raw duration and not INDEFINITELY_LOOPING_DURATION if looping
-		if (SoundDuration > 0.0f)
+		TObjectPtr<USoundWave> AudioComponentSoundWave = Cast<USoundWave>(AudioComponent->Sound);
+		if (AudioComponentSoundWave) 
 		{
-			ensure(CachedAudioStartTime >= 0.0f);
-			const float PlaybackPercentage = CachedAudioStartTime / SoundDuration;
-			return FGenericPlatformMath::Fmod(PlaybackPercentage, 1.0f);
+			const FWaveTransformUObjectConfiguration& SoundWaveTransformationChainConfig = AudioComponentSoundWave->GetTransformationChainConfig();
+			float ActiveDuration = SoundWaveTransformationChainConfig.EndTime - SoundWaveTransformationChainConfig.StartTime;
+			if (ActiveDuration <= 0.0f)
+			{
+				// If no active+initialized transformation, ActiveDuration will be <= 0.0f, so fallback to sound duration:
+				ActiveDuration = AudioComponent->Sound->Duration; // not GetDuration() to get the raw duration and not INDEFINITELY_LOOPING_DURATION if looping
+			}
+			if (ActiveDuration > 0.0f)
+			{
+				ensure(CachedAudioStartTime >= 0.0f);
+				const float PlaybackPercentage = CachedAudioStartTime / ActiveDuration;
+				CachedAudioStartTimeAsPercentage = FGenericPlatformMath::Fmod(PlaybackPercentage, 1.0f);
+			}
 		}
 	}
 	
-	return 0.0f;
+	return CachedAudioStartTimeAsPercentage;
 }
 
 const bool FWaveformEditorTransportController::SoundBaseIsValid() const
