Waveform Editor: Ability to tune loop points

diff --git a/Engine/Plugins/Editor/WaveformEditor/Source/WaveformTransformations/Private/WaveformTransformationMarkers.cpp b/Engine/Plugins/Editor/WaveformEditor/Source/WaveformTransformations/Private/WaveformTransformationMarkers.cpp
--- a/Engine/Plugins/Editor/WaveformEditor/Source/WaveformTransformations/Private/WaveformTransformationMarkers.cpp
+++ b/Engine/Plugins/Editor/WaveformEditor/Source/WaveformTransformations/Private/WaveformTransformationMarkers.cpp
@@ -117,3 +117,9 @@ void UWaveCueArray::Reset()
 	bIsInitialized = false;
 }
 #endif
+
+void UWaveCueArray::EnableLoopRegion(FSoundWaveCuePoint* OutSoundWaveCue)
+{
+	check(OutSoundWaveCue != nullptr);
+	OutSoundWaveCue->SetLoopRegion(true);
+}
diff --git a/Engine/Plugins/Editor/WaveformEditor/Source/WaveformTransformations/Public/WaveformTransformationMarkers.h b/Engine/Plugins/Editor/WaveformEditor/Source/WaveformTransformations/Public/WaveformTransformationMarkers.h
--- a/Engine/Plugins/Editor/WaveformEditor/Source/WaveformTransformations/Public/WaveformTransformationMarkers.h
+++ b/Engine/Plugins/Editor/WaveformEditor/Source/WaveformTransformations/Public/WaveformTransformationMarkers.h
@@ -21,6 +21,8 @@ public:
 	//Uninitialize and empty wavecue array
 	void Reset();
 
+	void EnableLoopRegion(FSoundWaveCuePoint* OutSoundWaveCue);
+
 	UPROPERTY(EditAnywhere, Category = "Markers")
 	TArray<FSoundWaveCuePoint> WaveCues;
 
diff --git a/Engine/Plugins/Editor/WaveformEditor/Source/WaveformTransformationsWidgets/Private/WaveformTransformationMarkersRenderer.cpp b/Engine/Plugins/Editor/WaveformEditor/Source/WaveformTransformationsWidgets/Private/WaveformTransformationMarkersRenderer.cpp
--- a/Engine/Plugins/Editor/WaveformEditor/Source/WaveformTransformationsWidgets/Private/WaveformTransformationMarkersRenderer.cpp
+++ b/Engine/Plugins/Editor/WaveformEditor/Source/WaveformTransformationsWidgets/Private/WaveformTransformationMarkersRenderer.cpp
@@ -3,10 +3,16 @@
 #include "WaveformTransformationMarkersRenderer.h"
 
 #include "Brushes/SlateRoundedBoxBrush.h"
+#include "Framework/Application/SlateApplication.h"
+#include "IWaveformTransformation.h"
 #include "Layout/Geometry.h"
 #include "PropertyHandle.h"
 #include "Rendering/DrawElements.h"
-#include "Widgets/SWidget.h"
+#include "Sound/SoundWave.h"
+#include "Widgets/Input/SButton.h"
+#include <Blueprint/WidgetLayoutLibrary.h>
+
+#define LOCTEXT_NAMESPACE "WaveformTransformationMarkerRenderer"
 
 FWaveformTransformationMarkerRenderer::FWaveformTransformationMarkerRenderer()
 {
@@ -81,10 +87,21 @@ int32 FWaveformTransformationMarkerRenderer::DrawMarkerHandles(const FGeometry&
 			LoopColors[ColorIndex],
 			false
 		);
+
+		FLinearColor LoopBoxColor = LoopColors[ColorIndex];
 		
 		//Make LoopBox always half as opaque as the Marker
-		FLinearColor LoopBoxColor = LoopColors[ColorIndex];
-		LoopBoxColor.A *= 0.5;
+		LoopBoxColor = LoopColors[ColorIndex];
+
+		float colorMultiplier = 0.5;
+
+		// Highlight selection
+		if (SelectedCue != INDEX_NONE && SelectedCue == LoopRegion.CuePointID)
+		{
+			colorMultiplier = 0.9;
+		}
+		
+		LoopBoxColor.A *= colorMultiplier;
 		FSlateDrawElement::MakeBox(
 			OutDrawElements,
 			++LayerId,
@@ -94,14 +111,18 @@ int32 FWaveformTransformationMarkerRenderer::DrawMarkerHandles(const FGeometry&
 			LoopBoxColor
 		);
 
-		FSlateDrawElement::MakeBox(
-			OutDrawElements,
-			++LayerId,
-			HandleGeometry,
-			&RoundedBoxBrush,
-			ESlateDrawEffect::None,
-			LoopColors[ColorIndex]
-		);
+		// If the loop region has no width or too small to draw handles, use the marker to move it
+		if (LoopRegion.FrameLength < LoopHandlePixelWidth * 2)
+		{
+			FSlateDrawElement::MakeBox(
+				OutDrawElements,
+				++LayerId,
+				HandleGeometry,
+				&RoundedBoxBrush,
+				ESlateDrawEffect::None,
+				LoopColors[ColorIndex]
+			);
+		}
 
 		FVector2D TextOffset(MarkerCenter + MarkerWidth, 0);
 		FSlateDrawElement::MakeText(
@@ -113,6 +134,39 @@ int32 FWaveformTransformationMarkerRenderer::DrawMarkerHandles(const FGeometry&
 			ESlateDrawEffect::None,
 			LabelTextColor
 		);
+
+		// If we have a loop region highlighted by the mouse, draw the handles
+		const bool DrawRegionHandles = (HighlightedCue != INDEX_NONE && HighlightedCue == LoopRegion.CuePointID && LoopRegion.FrameLength > LoopHandlePixelWidth);
+
+		if (DrawRegionHandles)
+		{
+			FPaintGeometry LeftHandleGeometry = AllottedGeometry.ToPaintGeometry(FVector2f(LoopHandlePixelWidth, AllottedGeometry.GetLocalSize().Y), FSlateLayoutTransform(FVector2f(LoopRegionPosition, 0)));
+			FPaintGeometry RightHandleGeometry = AllottedGeometry.ToPaintGeometry(FVector2f(LoopHandlePixelWidth, AllottedGeometry.GetLocalSize().Y), FSlateLayoutTransform(FVector2f(LoopRegionPosition + LoopRegionPixelLength - LoopHandlePixelWidth, 0)));
+			const bool bIsLeftHandleHighlighted = IsPositionInLoopHandleArea(MousePosition, AllottedGeometry, LoopRegion, true);
+			const bool bIsRightHandleHighlighted = IsPositionInLoopHandleArea(MousePosition, AllottedGeometry, LoopRegion, false);
+
+			FString HighlightHex = TEXT("#0078D7");
+			FColor HighlightColor = FColor::FromHex(HighlightHex);
+			FLinearColor HighlightLinearColor = FLinearColor::FromSRGBColor(HighlightColor);
+
+			FSlateDrawElement::MakeBox(
+				OutDrawElements,
+				++LayerId,
+				LeftHandleGeometry,
+				&Brush,
+				ESlateDrawEffect::None,
+				bIsLeftHandleHighlighted ? HighlightLinearColor : FLinearColor::White
+			);
+
+			FSlateDrawElement::MakeBox(
+				OutDrawElements,
+				++LayerId,
+				RightHandleGeometry,
+				&Brush,
+				ESlateDrawEffect::None,
+				bIsRightHandleHighlighted ? HighlightLinearColor : FLinearColor::White
+			);
+		}
 	}
 
 	for (const FSoundWaveCuePoint& Marker : CuePoints)
@@ -173,6 +227,10 @@ void FWaveformTransformationMarkerRenderer::Tick(const FGeometry& AllottedGeomet
 	check(NumFrames > 0);
 	PixelsPerFrame = FMath::Max(AllottedGeometry.GetLocalSize().X / NumFrames, 0.0);
 	ensure(PixelsPerFrame >= 0);
+	MarkerInInteractionRange.SetMinLoopSize((LoopHandlePixelWidth / PixelsPerFrame) * 3);
+
+	FVector2D MouseAbsolutePosition = UWidgetLayoutLibrary::GetMousePositionOnPlatform();
+	MousePosition = AllottedGeometry.AbsoluteToLocal(MouseAbsolutePosition);
 
 	MarkersArray = CastChecked<UWaveCueArray>(GetPropertyValue<UObject*>(MarkersHandle));
 	
@@ -226,20 +284,14 @@ FReply FWaveformTransformationMarkerRenderer::OnMouseButtonDown(SWidget& OwnerWi
 {
 	const FVector2D LocalCursorPosition = GetLocalCursorPosition(MouseEvent, MyGeometry);
 
-	const bool bIsLeftMouseButton = MouseEvent.IsMouseButtonDown(EKeys::LeftMouseButton);
+	MarkerInteractionType = EvaluateInteractionTypeFromCursorPosition(LocalCursorPosition, MyGeometry, MouseEvent.GetEffectingButton());
 
-	if (bIsLeftMouseButton)
+	if (MarkerInteractionType != EMarkerInteractionType::None)
 	{
-		MarkerInteractionType = GetInteractionTypeFromCursorPosition(LocalCursorPosition, MyGeometry);
-
-		if (MarkerInteractionType != EMarkerInteractionType::None)
-		{
-			MarkersHandle->NotifyPreChange();
-			BeginTransaction(TEXT("PropertyEditor"), FText::Format(NSLOCTEXT("PropertyEditor", "SetNumericPropertyTransaction", "Edit {0}"), FText::FromName(GET_MEMBER_NAME_CHECKED(UWaveformTransformationMarkers, Markers))), nullptr);
-
-			return FReply::Handled().CaptureMouse(OwnerWidget.AsShared()).PreventThrottling();
-		}
+		MarkersHandle->NotifyPreChange();
+		BeginTransaction(TEXT("PropertyEditor"), FText::Format(NSLOCTEXT("PropertyEditor", "SetNumericPropertyTransaction", "Edit {0}"), FText::FromName(GET_MEMBER_NAME_CHECKED(UWaveformTransformationMarkers, Markers))), nullptr);
 
+		return FReply::Handled().CaptureMouse(OwnerWidget.AsShared()).PreventThrottling();
 	}
 
 	return FReply::Unhandled();
@@ -247,6 +299,11 @@ FReply FWaveformTransformationMarkerRenderer::OnMouseButtonDown(SWidget& OwnerWi
 
 FReply FWaveformTransformationMarkerRenderer::OnMouseButtonUp(SWidget& OwnerWidget, const FGeometry& MyGeometry, const FPointerEvent& MouseEvent)
 {
+	if (TSharedPtr<SWindow> LockedFadeModeMenuWindow = CreateLoopMenuWindow.Pin())
+	{
+		FSlateApplication::Get().RequestDestroyWindow(LockedFadeModeMenuWindow.ToSharedRef());
+	}
+
 	if (MarkerInteractionType != EMarkerInteractionType::None)
 	{
 		MarkersArray->Modify();
@@ -271,23 +328,49 @@ FReply FWaveformTransformationMarkerRenderer::OnMouseMove(SWidget& OwnerWidget,
 		return FReply::Handled().CaptureMouse(OwnerWidget.AsShared());
 	}
 
+	// Evaluate highlight state for loop handles
+	HighlightedCue = INDEX_NONE;
+	for (const FSoundWaveCuePoint& CuePoint : MarkersArray->WaveCues)
+	{
+		if (CuePoint.IsLoopRegion() && IsPositionInLoopArea(MousePosition, MyGeometry, &CuePoint))
+		{
+			HighlightedCue = CuePoint.CuePointID;
+		}
+	}
+
 	return FReply::Unhandled();
 }
 
 FCursorReply FWaveformTransformationMarkerRenderer::OnCursorQuery(const FGeometry& MyGeometry, const FPointerEvent& CursorEvent) const
 {
-	if (MarkerInteractionType == EMarkerInteractionType::ScrubbingMarkerHandle)
+	// Grab hand if moving a regular marker
+	if (MarkerInteractionType == EMarkerInteractionType::MarkerHandle)
 	{
 		return FCursorReply::Cursor(EMouseCursor::GrabHandClosed);
 	}
 
-	const FVector2D LocalCursorPosition = GetLocalCursorPosition(CursorEvent, MyGeometry);
+	// Set resize on loop handle hover, loops are seperated as we want highlight interactions to superscede region highlights
+	for (const FSoundWaveCuePoint& CuePoint : MarkersArray->WaveCues)
+	{
+		const bool bIsLeftHandleHighlighted = IsPositionInLoopHandleArea(MousePosition, MyGeometry, CuePoint, true) || 
+			(MarkerInteractionType == EMarkerInteractionType::ScrubbingMarkerHandleLeft && HighlightedCue == CuePoint.CuePointID);
+		const bool bIsRightHandleHighlighted = IsPositionInLoopHandleArea(MousePosition, MyGeometry, CuePoint, false) || 
+			(MarkerInteractionType == EMarkerInteractionType::ScrubbingMarkerHandleRight && HighlightedCue == CuePoint.CuePointID);
 
-	// Set grab hand on mouse hover
+		if (CuePoint.IsLoopRegion() && (bIsLeftHandleHighlighted || bIsRightHandleHighlighted))
+		{
+			return FCursorReply::Cursor(EMouseCursor::ResizeLeftRight);
+		}
+		else if (!CuePoint.IsLoopRegion() && IsPositionInInteractionRange(CuePoint.FramePosition, MousePosition, MyGeometry))
+		{
+			return FCursorReply::Cursor(EMouseCursor::GrabHand);
+		}
+	}
+
+	// Set grab hand on loop region moving
 	for (const FSoundWaveCuePoint& CuePoint : MarkersArray->WaveCues)
 	{
-		if (IsPositionInInteractionRange(CuePoint.FramePosition, LocalCursorPosition, MyGeometry) ||
-			(CuePoint.IsLoopRegion() && IsPositionInInteractionRange(CuePoint.FramePosition + CuePoint.FrameLength, LocalCursorPosition, MyGeometry)))
+		if (CuePoint.IsLoopRegion() && IsPositionInLoopArea(MousePosition, MyGeometry, &CuePoint))
 		{
 			return FCursorReply::Cursor(EMouseCursor::GrabHand);
 		}
@@ -349,9 +432,60 @@ bool FWaveformTransformationMarkerRenderer::IsPositionInInteractionRange(const i
 	return false;
 }
 
+bool FWaveformTransformationMarkerRenderer::IsPositionInLoopArea(const FVector2D& InLocalCursorPosition, const FGeometry& WidgetGeometry, const FSoundWaveCuePoint* CueMarker, bool IncludeHandleArea) const
+{
+	check(TransformationWaveInfo.NumChannels > 0);
+	check(CueMarker != nullptr);
+	ensure(LoopHandlePixelWidth > 0);
+
+	const uint32 StartFrameOffset = TransformationWaveInfo.StartFrameOffset / TransformationWaveInfo.NumChannels;
+	float LoopRegionPosition = (CueMarker->FramePosition - StartFrameOffset) * PixelsPerFrame;
+	float LoopRegionPixelLength = CueMarker->FrameLength * PixelsPerFrame;
+
+	// If we want the handles to be included for the area check, as having the mouse in the handle area has different behaviour to the total loop area
+	if (!IncludeHandleArea)
+	{
+		LoopRegionPosition += LoopHandlePixelWidth;
+		LoopRegionPixelLength -= LoopHandlePixelWidth * 2;
+	}
+
+	check(LoopRegionPixelLength >= 0);
+
+	if (InLocalCursorPosition.X >= LoopRegionPosition && InLocalCursorPosition.X <= LoopRegionPosition + LoopRegionPixelLength &&
+		InLocalCursorPosition.Y <= WidgetGeometry.GetLocalSize().Y && InLocalCursorPosition.Y >= 0)
+	{
+		return true;
+	}
+
+	return false;
+}
+
+bool FWaveformTransformationMarkerRenderer::IsPositionInLoopHandleArea(const FVector2D& InLocalCursorPosition, const FGeometry& WidgetGeometry, const FSoundWaveCuePoint& CueMarker, bool IsLeftHandle) const
+{
+	check(TransformationWaveInfo.NumChannels > 0);
+	ensure(LoopHandlePixelWidth > 0);
+
+	const uint32 StartFrameOffset = TransformationWaveInfo.StartFrameOffset / TransformationWaveInfo.NumChannels;
+	float LoopRegionPosition = IsLeftHandle ?
+		(CueMarker.FramePosition - StartFrameOffset) * PixelsPerFrame :
+		((CueMarker.FramePosition + CueMarker.FrameLength) * PixelsPerFrame) - LoopHandlePixelWidth;
+
+	// Check that only the handle within the loop area is highlighted
+	if (InLocalCursorPosition.X > LoopRegionPosition && InLocalCursorPosition.X < LoopRegionPosition + LoopHandlePixelWidth &&
+		InLocalCursorPosition.Y < WidgetGeometry.GetLocalSize().Y && InLocalCursorPosition.Y > 0)
+	{
+		return true;
+	}
+
+	return false;
+}
+
 //Returns true if there is a marker in interaction range
 bool FWaveformTransformationMarkerRenderer::SetMarkerInInteractionRange(const FVector2D& InLocalCursorPosition, const FGeometry& WidgetGeometry)
 {
+	check(TransformationWaveInfo.NumChannels > 0);
+	ensure(LoopHandlePixelWidth > 0);
+
 	InteractionRanges.Reset(MarkersArray->WaveCues.Num());
 	
 	for (FSoundWaveCuePoint& CuePoint : MarkersArray->WaveCues)
@@ -368,20 +502,70 @@ bool FWaveformTransformationMarkerRenderer::SetMarkerInInteractionRange(const FV
 			continue;
 		}
 
-		const int32 LoopRegionEndHandlePosition = CuePoint.FramePosition + CuePoint.FrameLength;
-		
-		if (IsPositionInInteractionRange(LoopRegionEndHandlePosition, InLocalCursorPosition, WidgetGeometry))
+		// Check highlights of handles independantly first over loop area
+		if (IsPositionInLoopHandleArea(InLocalCursorPosition, WidgetGeometry, CuePoint, true) || IsPositionInLoopHandleArea(InLocalCursorPosition, WidgetGeometry, CuePoint, false))
 		{
 			MarkerInInteractionRange.SetMarkerInInteractionRange(&CuePoint, true);
-			
+
+			return true;
+		}
+	}
+
+	check(PixelsPerFrame > 0);
+
+	// Evaluate loop area if no handles are selected
+	for (FSoundWaveCuePoint& CuePoint : MarkersArray->WaveCues)
+	{
+		if (CuePoint.FrameLength > 0 && IsPositionInLoopArea(InLocalCursorPosition, WidgetGeometry, &CuePoint))
+		{
+			const uint32 StartFrameOffset = TransformationWaveInfo.StartFrameOffset / TransformationWaveInfo.NumChannels;
+			int32 LoopRegionPosition = (InLocalCursorPosition.X / PixelsPerFrame) - StartFrameOffset - CuePoint.FramePosition;
+
+			MarkerInInteractionRange.SetMarkerInInteractionRange(&CuePoint, false, LoopRegionPosition);
+			SelectedCue = CuePoint.CuePointID;
+
 			return true;
 		}
 	}
+
+	// MarkerPtr is set to null if there are no markers in range to handle cases where the mouse is not over any markers
 	MarkerInInteractionRange.SetMarkerInInteractionRange(nullptr);
 	
 	return false;
 }
 
+bool FWaveformTransformationMarkerRenderer::SelectLoopArea(const FVector2D& InLocalCursorPosition, const FGeometry& WidgetGeometry)
+{
+	check(TransformationWaveInfo.NumChannels > 0);
+	ensure(LoopHandlePixelWidth > 0);
+
+	const uint32 StartFrameOffset = TransformationWaveInfo.StartFrameOffset / TransformationWaveInfo.NumChannels;
+
+	for (FSoundWaveCuePoint& LoopRegion : LoopPoints)
+	{
+		// Check mouse position bounds
+		if (IsPositionInLoopArea(InLocalCursorPosition, WidgetGeometry, &LoopRegion))
+		{
+			SetActiveLoopRegion(LoopRegion.CuePointID);
+
+			return true;
+		}
+	}
+
+	SetActiveLoopRegion(INDEX_NONE);
+
+	return false;
+}
+
+void FWaveformTransformationMarkerRenderer::SetActiveLoopRegion(int32 CuePointID)
+{
+	MarkersHandle->NotifyPreChange();
+	MarkersArray->Modify();
+	SelectedCue = CuePointID;
+	EndTransaction();
+	MarkersHandle->NotifyPostChange(EPropertyChangeType::ValueSet);
+}
+
 FVector2D FWaveformTransformationMarkerRenderer::GetLocalCursorPosition(const FPointerEvent& MouseEvent, const FGeometry& EventGeometry)
 {
 	const FVector2D ScreenSpacePosition = MouseEvent.GetScreenSpacePosition();
@@ -389,11 +573,50 @@ FVector2D FWaveformTransformationMarkerRenderer::GetLocalCursorPosition(const FP
 	return  EventGeometry.AbsoluteToLocal(ScreenSpacePosition);
 }
 
-FWaveformTransformationMarkerRenderer::EMarkerInteractionType FWaveformTransformationMarkerRenderer::GetInteractionTypeFromCursorPosition(const FVector2D& InLocalCursorPosition, const FGeometry& WidgetGeometry)
+FWaveformTransformationMarkerRenderer::EMarkerInteractionType FWaveformTransformationMarkerRenderer::EvaluateInteractionTypeFromCursorPosition(const FVector2D& InLocalCursorPosition, const FGeometry& WidgetGeometry, const FKey MouseButton)
 {
+	if (MouseButton == EKeys::RightMouseButton)
+	{
+		if (SetMarkerInInteractionRange(InLocalCursorPosition, WidgetGeometry) && MarkerInInteractionRange.GetMarker() != nullptr && MarkerInInteractionRange.GetMarker()->IsLoopRegion() == false)
+		{
+			return EMarkerInteractionType::RightClickMarker;
+		}
+	}
+
+	if (MouseButton != EKeys::LeftMouseButton)
+	{
+		return EMarkerInteractionType::None;
+	}
+
+	for (FSoundWaveCuePoint CuePoint : MarkersArray->WaveCues)
+	{
+		if (!CuePoint.IsLoopRegion())
+		{
+			continue;
+		}			
+			
+		if (IsPositionInLoopHandleArea(InLocalCursorPosition, WidgetGeometry, CuePoint, true))
+		{
+			SetMarkerInInteractionRange(InLocalCursorPosition, WidgetGeometry);
+			return EMarkerInteractionType::ScrubbingMarkerHandleLeft;
+		}
+
+		if (IsPositionInLoopHandleArea(InLocalCursorPosition, WidgetGeometry, CuePoint, false))
+		{
+			SetMarkerInInteractionRange(InLocalCursorPosition, WidgetGeometry);
+			return EMarkerInteractionType::ScrubbingMarkerHandleRight;
+		}
+	}
+
 	if (SetMarkerInInteractionRange(InLocalCursorPosition, WidgetGeometry))
 	{
-		return EMarkerInteractionType::ScrubbingMarkerHandle;
+		return EMarkerInteractionType::MarkerHandle;
+	}
+	
+	if (SelectLoopArea(InLocalCursorPosition, WidgetGeometry))
+	{
+		SetMarkerInInteractionRange(InLocalCursorPosition, WidgetGeometry);
+		return EMarkerInteractionType::LoopHandle;
 	}
 
 	return EMarkerInteractionType::None;
@@ -401,38 +624,134 @@ FWaveformTransformationMarkerRenderer::EMarkerInteractionType FWaveformTransform
 
 void FWaveformTransformationMarkerRenderer::SetPropertyValueDependingOnInteractionType(const FPointerEvent& MouseEvent, const FGeometry& WidgetGeometry)
 {
+	ensure(PixelsPerFrame > 0);
+
 	const FVector2D LocalCursorPosition = GetLocalCursorPosition(MouseEvent, WidgetGeometry);
 	const float LocalCursorXClamped = FMath::Clamp(LocalCursorPosition.X, 0.f, WidgetGeometry.GetLocalSize().X);
 
-	if (MarkerInteractionType == EMarkerInteractionType::ScrubbingMarkerHandle)
+	switch (MarkerInteractionType)
+	{
+	case EMarkerInteractionType::None:
+		break;
+	case EMarkerInteractionType::MarkerHandle:
+		{
+			check(PixelsPerFrame > 0);
+			MarkerInInteractionRange.SetMarkerPosition(LocalCursorXClamped / PixelsPerFrame, false);
+		}
+		break;
+	case EMarkerInteractionType::ScrubbingMarkerHandleRight:
+		{
+			if (MarkerInInteractionRange.GetMarker())
+			{
+				check(PixelsPerFrame > 0);
+				MarkerInInteractionRange.SetMarkerPosition(LocalCursorXClamped / PixelsPerFrame, false);
+				HighlightedCue = MarkerInInteractionRange.GetMarker()->CuePointID;
+			}
+		}
+		break;
+	case EMarkerInteractionType::ScrubbingMarkerHandleLeft:
+	{
+		if (MarkerInInteractionRange.GetMarker())
+		{
+			check(PixelsPerFrame > 0);
+			MarkerInInteractionRange.SetMarkerPosition(LocalCursorXClamped / PixelsPerFrame, true);
+			HighlightedCue = MarkerInInteractionRange.GetMarker()->CuePointID;
+		}
+	}
+		break;
+	case EMarkerInteractionType::LoopHandle:
 	{
 		if (MarkerInInteractionRange.GetMarker())
 		{
 			check(PixelsPerFrame > 0);
-			MarkerInInteractionRange.SetMarkerPosition(LocalCursorXClamped / PixelsPerFrame);
+			MarkerInInteractionRange.SetMarkerPosition(LocalCursorXClamped / PixelsPerFrame, true);
 		}
 	}
+		break;
+	case EMarkerInteractionType::RightClickMarker:
+	{
+		ShowConvertToLoopAtMarker(WidgetGeometry, MouseEvent, MarkerInInteractionRange.GetMarker());
+	}
+		break;
+	}
+}
+
+void FWaveformTransformationMarkerRenderer::ShowConvertToLoopAtMarker(const FGeometry& MyGeometry, const FPointerEvent& MouseEvent, FSoundWaveCuePoint* CueToModify)
+{
+	check(CueToModify);
+	const FVector2D LocalCursorPosition = MouseEvent.GetScreenSpacePosition();
+
+	TSharedRef<SVerticalBox> MenuContent = SNew(SVerticalBox);
+
+	if (TSharedPtr<SWindow> LockedFadeModeMenuWindow = CreateLoopMenuWindow.Pin())
+	{
+		FSlateApplication::Get().RequestDestroyWindow(LockedFadeModeMenuWindow.ToSharedRef());
+	}
+
+	MenuContent->AddSlot()
+		.AutoHeight()
+		[
+			SNew(SButton)
+				.Text(LOCTEXT("WaveformTransformationMarkerRendererConvertLoop","Convert To Loop Region"))
+				.OnClicked_Lambda([this, CueToModify]() -> FReply
+					{
+						//ToggleLoopRegion(CueToModify);
+						MarkersArray->EnableLoopRegion(CueToModify);
+						CueToModify->FrameLength = MarkerInInteractionRange.GetMinLoopSize();
+
+						if (TSharedPtr<SWindow> LockedFadeModeMenuWindow = CreateLoopMenuWindow.Pin())
+						{
+							FSlateApplication::Get().RequestDestroyWindow(LockedFadeModeMenuWindow.ToSharedRef());
+						}
+						return FReply::Handled();
+					})
+		];
+
+	TSharedRef<SWindow> MenuWindow = SNew(SWindow)
+		.AutoCenter(EAutoCenter::None)
+		.ScreenPosition(LocalCursorPosition)
+		.SizingRule(ESizingRule::Autosized)
+		.SupportsTransparency(EWindowTransparency::None)
+		.IsPopupWindow(true)
+		.CreateTitleBar(false)
+		[
+			MenuContent
+		];
+
+	CreateLoopMenuWindow = MenuWindow;
+
+	FSlateApplication::Get().AddWindow(MenuWindow);
 }
 
-void FWaveformTransformationMarkerRenderer::FMarkerInInteractionRange::SetMarkerInInteractionRange(FSoundWaveCuePoint* InMarkerPtr, const bool InIsLoopRegionEndMarker)
+void FWaveformTransformationMarkerRenderer::FMarkerInInteractionRange::SetMarkerInInteractionRange(FSoundWaveCuePoint* InMarkerPtr, const bool InIsLoopRegionEndMarker, int32 GrabXOffset)
 {
 	MarkerPtr = InMarkerPtr;
 	IsLoopRegionEndMarker = InIsLoopRegionEndMarker;
+	GrabOffset = GrabXOffset;
 }
 
-void FWaveformTransformationMarkerRenderer::FMarkerInInteractionRange::SetMarkerPosition(const float InPosition)
+void FWaveformTransformationMarkerRenderer::FMarkerInInteractionRange::SetMarkerPosition(const float InPosition, bool bIsLeft)
 {
 	check(MarkerPtr);
 
 	if (IsLoopRegionEndMarker)
 	{
-		constexpr int32 MinLoopSize = 10;
-		
 		//FMath::Max prevents the user from collapsing the loop
-		MarkerPtr->FrameLength = FMath::Max(InPosition, MarkerPtr->FramePosition + MinLoopSize) - MarkerPtr->FramePosition;
+		if (bIsLeft)
+		{
+			int32 MoveDifference = MarkerPtr->FramePosition - (int32)InPosition;
+			MarkerPtr->FramePosition = (int32)InPosition;
+			MarkerPtr->FrameLength = FMath::Max(MarkerPtr->FrameLength + MoveDifference, MinLoopSize);
+		}
+		else
+		{
+			MarkerPtr->FrameLength = FMath::Max(InPosition, MarkerPtr->FramePosition + MinLoopSize) - MarkerPtr->FramePosition;
+		}
 	}
 	else
 	{
-		MarkerPtr->FramePosition = InPosition;
+		MarkerPtr->FramePosition = FMath::Max(0, InPosition - GrabOffset);
 	}
 }
+
+#undef LOCTEXT_NAMESPACE
\ No newline at end of file
diff --git a/Engine/Plugins/Editor/WaveformEditor/Source/WaveformTransformationsWidgets/Private/WaveformTransformationMarkersRenderer.h b/Engine/Plugins/Editor/WaveformEditor/Source/WaveformTransformationsWidgets/Private/WaveformTransformationMarkersRenderer.h
--- a/Engine/Plugins/Editor/WaveformEditor/Source/WaveformTransformationsWidgets/Private/WaveformTransformationMarkersRenderer.h
+++ b/Engine/Plugins/Editor/WaveformEditor/Source/WaveformTransformationsWidgets/Private/WaveformTransformationMarkersRenderer.h
@@ -2,6 +2,8 @@
 
 #pragma once
 
+#include "IWaveformTransformation.h"
+#include "Sound/SoundWave.h"
 #include "WaveformTransformationMarkers.h"
 #include "WaveformTransformationRendererBase.h"
 #include "WaveformTransformationsWidgetsSettings.h"
@@ -29,7 +31,11 @@ private:
 	void SetLoopColors(const TArray<FLinearColor>& InColors);
 
 	bool IsPositionInInteractionRange(const int32 InFramePosition, const FVector2D& InLocalCursorPosition, const FGeometry& WidgetGeometry) const;
+	bool IsPositionInLoopArea(const FVector2D& InLocalCursorPosition, const FGeometry& WidgetGeometry, const FSoundWaveCuePoint* CueMarker, bool IncludeHandleArea = true) const;
+	bool IsPositionInLoopHandleArea(const FVector2D& InLocalCursorPosition, const FGeometry& WidgetGeometry, const FSoundWaveCuePoint& CueMarker, bool IsLeftHandle) const;
 	bool SetMarkerInInteractionRange(const FVector2D& InLocalCursorPosition, const FGeometry& WidgetGeometry);
+	bool SelectLoopArea(const FVector2D& InLocalCursorPosition, const FGeometry& WidgetGeometry);
+	void SetActiveLoopRegion(int32 CuePointID);
 	static FVector2D GetLocalCursorPosition(const FPointerEvent& MouseEvent, const FGeometry& EventGeometry);
 
 	TArray<TRange<float>> InteractionRanges;
@@ -37,24 +43,34 @@ private:
 	enum class EMarkerInteractionType : uint8
 	{
 		None = 0,
-		ScrubbingMarkerHandle
+		ScrubbingMarkerHandleRight,
+		ScrubbingMarkerHandleLeft,
+		LoopHandle,
+		MarkerHandle,
+		RightClickMarker,
 	} MarkerInteractionType = EMarkerInteractionType::None;
 
 	class FMarkerInInteractionRange
 	{
 	public:
-		void SetMarkerInInteractionRange(FSoundWaveCuePoint* InMarkerPtr, const bool InIsLoopRegionEndMarker = false);
-		const FSoundWaveCuePoint* GetMarker() const { return MarkerPtr; }
-		void SetMarkerPosition(const float InPosition);
+		void SetMarkerInInteractionRange(FSoundWaveCuePoint* InMarkerPtr, const bool InIsLoopRegionEndMarker = false, int32 GrabXOffset = 0);
+		FSoundWaveCuePoint* GetMarker() const { return MarkerPtr; }
+		void SetMarkerPosition(const float InPosition, bool bIsLeft);
+		void SetMinLoopSize(int32 value) { check(value > 0); MinLoopSize = value; }
+		const int32 GetMinLoopSize() { return MinLoopSize; }
 		
 	private:
 		FSoundWaveCuePoint* MarkerPtr = nullptr;
 		bool IsLoopRegionEndMarker = false;
+		int32 GrabOffset = 0;
+		int32 MinLoopSize = 10;
 	} MarkerInInteractionRange;
 
-	EMarkerInteractionType GetInteractionTypeFromCursorPosition(const FVector2D& InLocalCursorPosition, const FGeometry& WidgetGeometry);
+	EMarkerInteractionType EvaluateInteractionTypeFromCursorPosition(const FVector2D& InLocalCursorPosition, const FGeometry& WidgetGeometry, const FKey MouseButton);
 	void SetPropertyValueDependingOnInteractionType(const FPointerEvent& MouseEvent, const FGeometry& WidgetGeometry);
 
+	void ShowConvertToLoopAtMarker(const FGeometry& MyGeometry, const FPointerEvent& MouseEvent, FSoundWaveCuePoint* CueToModify);
+
 	TArray<FSoundWaveCuePoint> CuePoints;
 	TArray<FSoundWaveCuePoint> LoopPoints;
 
@@ -62,11 +78,19 @@ private:
 
 	double PixelsPerFrame = 0.0;
 
+	// Ensures handles are always visible to a user and we can avoid overlapping handle regions when loop regions are too small, in future consider dynamic relative size based on window size
+	static constexpr float LoopHandlePixelWidth = 7.0f; 
+
+	FVector2D MousePosition;
+	FSoundWaveCuePoint* MarkerToModify = nullptr;
+	int32 SelectedCue = INDEX_NONE;
+	int32 HighlightedCue = INDEX_NONE;
+
 	FLinearColor MarkerColor = WaveformTransformationWidgetSharedDefaults::DefaultMarkerColor;
 	TArray<FLinearColor> LoopColors = WaveformTransformationWidgetSharedDefaults::DefaultLoopColors;
 	FLinearColor LabelTextColor = WaveformTransformationWidgetSharedDefaults::DefaultLabelTextColor;
 	float LabelFontSize = WaveformTransformationWidgetSharedDefaults::DefaultLabelFontSize;
 
 	TSharedPtr<IPropertyHandle> MarkersHandle;
-
+	TWeakPtr<SWindow> CreateLoopMenuWindow;
 };
diff --git a/Engine/Plugins/Editor/WaveformEditor/Source/WaveformTransformationsWidgets/WaveformTransformationsWidgets.build.cs b/Engine/Plugins/Editor/WaveformEditor/Source/WaveformTransformationsWidgets/WaveformTransformationsWidgets.build.cs
--- a/Engine/Plugins/Editor/WaveformEditor/Source/WaveformTransformationsWidgets/WaveformTransformationsWidgets.build.cs
+++ b/Engine/Plugins/Editor/WaveformEditor/Source/WaveformTransformationsWidgets/WaveformTransformationsWidgets.build.cs
@@ -26,6 +26,7 @@ namespace UnrealBuildTool.Rules
 					"SlateCore", 
 					"ToolMenus", 
 					"UnrealEd",
+					"UMG",
 					"WaveformTransformations",
 				}
 			);
diff --git a/Engine/Source/Runtime/AudioExtensions/Public/IWaveformTransformation.h b/Engine/Source/Runtime/AudioExtensions/Public/IWaveformTransformation.h
--- a/Engine/Source/Runtime/AudioExtensions/Public/IWaveformTransformation.h
+++ b/Engine/Source/Runtime/AudioExtensions/Public/IWaveformTransformation.h
@@ -74,6 +74,7 @@ struct FSoundWaveCuePoint
 	int32 FrameLength = 0;
 
 	bool IsLoopRegion() const { return bIsLoopRegion; }
+	void SetLoopRegion(bool value) { bIsLoopRegion = value; }
 
 #if WITH_EDITORONLY_DATA
 	void ScaleFrameValues(float Factor)
@@ -85,6 +86,7 @@ struct FSoundWaveCuePoint
 
 	friend class USoundFactory;
 	friend class USoundWave;
+
 private:
 	// intentionally kept private.
 	// only USoundFactory should modify this value on import
