[Waveform Editor] Bug Fix - Sound Waves are not marked dirty in Content Browser with Default Transforms in Waveform Editor

diff --git a/Engine/Plugins/Editor/WaveformEditor/Source/WaveformEditor/Private/WaveformEditor.cpp b/Engine/Plugins/Editor/WaveformEditor/Source/WaveformEditor/Private/WaveformEditor.cpp
--- a/Engine/Plugins/Editor/WaveformEditor/Source/WaveformEditor/Private/WaveformEditor.cpp
+++ b/Engine/Plugins/Editor/WaveformEditor/Source/WaveformEditor/Private/WaveformEditor.cpp
@@ -145,15 +145,12 @@ void FWaveformEditor::AddDefaultTransformations()
 			{
 				SoundWave->Transformations.Add(nullptr);
 			}
+
+			//Update the content browser asset state to dirty
+			SoundWave->PostEditChange();
 		}
 
 		TransformationChainConfig = SoundWave->UpdateTransformations();
-
-		//Update the content browser asset state to dirty
-		if (GEditor && !TransformationsSettings->LaunchTransformations.IsEmpty())
-		{
-			GEditor->BroadcastObjectReimported(SoundWave);
-		}
 	}
 }
 
@@ -172,10 +169,7 @@ void FWaveformEditor::NotifyPostTransformationChange(const EPropertyChangeType::
 	NotifyPostChange(PropertyChangedEvent, &PropertyChain);
 
 	//Update the content browser asset state to dirty
-	if (GEditor)
-	{
-		GEditor->BroadcastObjectReimported(SoundWave);
-	}
+	SoundWave->PostEditChange();
 
 	TransformationChainConfig = SoundWave->UpdateTransformations();
 }
