Waveform Editor: Loop point nudging shortcuts

diff --git a/Engine/Plugins/Editor/WaveformEditor/Source/WaveformEditor/Private/WaveformEditor.cpp b/Engine/Plugins/Editor/WaveformEditor/Source/WaveformEditor/Private/WaveformEditor.cpp
--- a/Engine/Plugins/Editor/WaveformEditor/Source/WaveformEditor/Private/WaveformEditor.cpp
+++ b/Engine/Plugins/Editor/WaveformEditor/Source/WaveformEditor/Private/WaveformEditor.cpp
@@ -25,6 +25,7 @@
 #include "WaveformEditorToolMenuContext.h"
 #include "WaveformEditorTransformationsSettings.h"
 #include "WaveformEditorWaveWriter.h"
+#include "WaveformTransformationMarkers.h"
 #include "Widgets/Docking/SDockTab.h"
 
 #define LOCTEXT_NAMESPACE "WaveformEditor"
@@ -286,6 +287,59 @@ bool FWaveformEditor::CanFadeOut()
 	return true;
 }
 
+TObjectPtr<UWaveformTransformationMarkers> FWaveformEditor::GetOrAddMarkerTransformation()
+{
+	check(SoundWave);
+
+	TObjectPtr<UWaveformTransformationMarkers> MarkersTransformation;
+
+	for (TObjectPtr<UWaveformTransformationBase> Transformation : SoundWave->Transformations)
+	{
+		MarkersTransformation = Cast<UWaveformTransformationMarkers>(Transformation);
+		if (MarkersTransformation)
+		{
+			break;
+		}
+	}
+
+	if (MarkersTransformation == nullptr)
+	{
+		EObjectFlags MaskedOuterFlags = SoundWave ? SoundWave->GetMaskedFlags(RF_PropagateToSubObjects) : RF_NoFlags;
+
+		if (SoundWave->HasAnyFlags(RF_ClassDefaultObject | RF_ArchetypeObject))
+		{
+			MaskedOuterFlags |= RF_ArchetypeObject;
+		}
+
+		MarkersTransformation = NewObject<UWaveformTransformationMarkers>(SoundWave, UWaveformTransformationMarkers::StaticClass(), NAME_None, MaskedOuterFlags);
+
+		SoundWave->Modify();
+		SoundWave->Transformations.Add(MarkersTransformation);
+
+		NotifyPostTransformationChange();
+	}
+
+	return MarkersTransformation;
+}
+
+void FWaveformEditor::ModifyMarkerLoopRegion(ELoopModificationControls Modification)
+{
+	check(SoundWave);
+	TObjectPtr<UWaveformTransformationMarkers> MarkersTransformation = GetOrAddMarkerTransformation();
+	check(MarkersTransformation);
+
+	MarkersTransformation->ModifyMarkerLoopRegion(Modification);
+}
+
+void FWaveformEditor::CycleMarkerLoopRegion(ELoopModificationControls Modification)
+{
+	check(SoundWave);
+	TObjectPtr<UWaveformTransformationMarkers> MarkersTransformation = GetOrAddMarkerTransformation();
+	check(MarkersTransformation);
+
+	MarkersTransformation->CycleMarkerLoopRegion(Modification);
+}
+
 bool FWaveformEditor::InitializeAudioComponent()
 {
 	if (!ensure(SoundWave))
@@ -865,6 +919,38 @@ bool FWaveformEditor::BindCommands()
 		FCanExecuteAction::CreateLambda([this] {return true; }),
 		FIsActionChecked::CreateLambda([this] {return FadeOutCurve == UWaveformTransformationTrimFade::FadeModeToCurveValueMap[EWaveEditorFadeMode::Sigmoid]; }));
 
+	ToolkitCommands->MapAction(
+		Commands.LeftBoundsIncrease,
+		FExecuteAction::CreateLambda([this] { ModifyMarkerLoopRegion(ELoopModificationControls::LeftHandleIncrement); }));
+
+	ToolkitCommands->MapAction(
+		Commands.LeftBoundsDecrease,
+		FExecuteAction::CreateLambda([this] { ModifyMarkerLoopRegion(ELoopModificationControls::LeftHandleDecrement); }));
+
+	ToolkitCommands->MapAction(
+		Commands.RightBoundsIncrease,
+		FExecuteAction::CreateLambda([this] { ModifyMarkerLoopRegion(ELoopModificationControls::RightHandleIncrement); }));
+
+	ToolkitCommands->MapAction(
+		Commands.RightBoundsDecrease,
+		FExecuteAction::CreateLambda([this] { ModifyMarkerLoopRegion(ELoopModificationControls::RightHandleDecrement); }));
+
+	ToolkitCommands->MapAction(
+		Commands.BoundsIncrementIncrease,
+		FExecuteAction::CreateLambda([this] { ModifyMarkerLoopRegion(ELoopModificationControls::IncreaseIncrement); }));
+
+	ToolkitCommands->MapAction(
+		Commands.BoundsIncrementDecrease,
+		FExecuteAction::CreateLambda([this] { ModifyMarkerLoopRegion(ELoopModificationControls::DecreaseIncrement); }));
+
+	ToolkitCommands->MapAction(
+		Commands.SelectNextLoop,
+		FExecuteAction::CreateLambda([this] { CycleMarkerLoopRegion(ELoopModificationControls::SelectNextLoop); }));
+
+	ToolkitCommands->MapAction(
+		Commands.SelectPreviousLoop,
+		FExecuteAction::CreateLambda([this] { CycleMarkerLoopRegion(ELoopModificationControls::SelectPreviousLoop); }));
+
 	return true;
 }
 
diff --git a/Engine/Plugins/Editor/WaveformEditor/Source/WaveformEditor/Private/WaveformEditorCommands.cpp b/Engine/Plugins/Editor/WaveformEditor/Source/WaveformEditor/Private/WaveformEditorCommands.cpp
--- a/Engine/Plugins/Editor/WaveformEditor/Source/WaveformEditor/Private/WaveformEditorCommands.cpp
+++ b/Engine/Plugins/Editor/WaveformEditor/Source/WaveformEditor/Private/WaveformEditorCommands.cpp
@@ -44,6 +44,15 @@ void FWaveformEditorCommands::RegisterCommands()
 	UI_COMMAND(FadeOutLogarithmic, "Logarithmic Fade Out", "Fade Out Curve type applied when Toggle Fade Out pressed", EUserInterfaceActionType::RadioButton, FInputChord());
 	UI_COMMAND(FadeOutSigmoid, "Sigmoid Fade Out", "Fade Out S-Shaped Curve type applied when Toggle Fade Out pressed", EUserInterfaceActionType::RadioButton, FInputChord());
 
+	UI_COMMAND(LeftBoundsIncrease, "Increase Left Loop Bounds", "Moves the start of the loop region earlier", EUserInterfaceActionType::Button, FInputChord(EKeys::Comma));
+	UI_COMMAND(LeftBoundsDecrease, "Decrease Left Loop Bounds", "Moves the start of the loop region later", EUserInterfaceActionType::Button, FInputChord(EKeys::Period));
+	UI_COMMAND(RightBoundsIncrease, "Increase Right Loop Bounds", "Moves the end of the loop region earlier", EUserInterfaceActionType::Button, FInputChord(EKeys::RightBracket));
+	UI_COMMAND(RightBoundsDecrease, "Decrease Right Loop Bounds", "Moves the end of the loop region later", EUserInterfaceActionType::Button, FInputChord(EKeys::LeftBracket));
+
+	UI_COMMAND(BoundsIncrementIncrease, "Increase Bound Increment", "Increases the increment to shift the loop bounds", EUserInterfaceActionType::Button, FInputChord(EKeys::Equals));
+	UI_COMMAND(BoundsIncrementDecrease, "Decrease Bound Increment", "Increases the increment to shift the loop bounds", EUserInterfaceActionType::Button, FInputChord(EKeys::Hyphen));
+	UI_COMMAND(SelectNextLoop, "Select Next Loop", "Selects the next loop marker", EUserInterfaceActionType::Button, FInputChord(EKeys::F));
+	UI_COMMAND(SelectPreviousLoop, "Select Previous Loop", "Selects the previous loop marker", EUserInterfaceActionType::Button, FInputChord(EKeys::B));
 }
 	
 const FWaveformEditorCommands& FWaveformEditorCommands::Get()
diff --git a/Engine/Plugins/Editor/WaveformEditor/Source/WaveformEditor/Public/WaveformEditor.h b/Engine/Plugins/Editor/WaveformEditor/Source/WaveformEditor/Public/WaveformEditor.h
--- a/Engine/Plugins/Editor/WaveformEditor/Source/WaveformEditor/Public/WaveformEditor.h
+++ b/Engine/Plugins/Editor/WaveformEditor/Source/WaveformEditor/Public/WaveformEditor.h
@@ -10,6 +10,7 @@
 #include "WaveformEditorZoomController.h"
 #include "WaveformTransformationTrimFade.h"
 #include "TransformedWaveformView.h"
+#include <WaveformTransformationMarkers.h>
 
 class IToolkitHost;
 class SDockTab;
@@ -111,6 +112,10 @@ private:
 	void ToggleFadeOut();
 	bool CanFadeOut();
 
+	TObjectPtr<UWaveformTransformationMarkers> GetOrAddMarkerTransformation();
+	void ModifyMarkerLoopRegion(ELoopModificationControls Modification);
+	void CycleMarkerLoopRegion(ELoopModificationControls Modification);
+
 	FTransformedWaveformView  WaveformView;
 
 	/** Exports the edited waveform to a new asset */
diff --git a/Engine/Plugins/Editor/WaveformEditor/Source/WaveformEditor/Public/WaveformEditorCommands.h b/Engine/Plugins/Editor/WaveformEditor/Source/WaveformEditor/Public/WaveformEditorCommands.h
--- a/Engine/Plugins/Editor/WaveformEditor/Source/WaveformEditor/Public/WaveformEditorCommands.h
+++ b/Engine/Plugins/Editor/WaveformEditor/Source/WaveformEditor/Public/WaveformEditorCommands.h
@@ -43,4 +43,13 @@ public:
 	TSharedPtr<FUICommandInfo> FadeOutExponential;
 	TSharedPtr<FUICommandInfo> FadeOutLogarithmic;
 	TSharedPtr<FUICommandInfo> FadeOutSigmoid;
+
+	TSharedPtr<FUICommandInfo> LeftBoundsIncrease;
+	TSharedPtr<FUICommandInfo> LeftBoundsDecrease;
+	TSharedPtr<FUICommandInfo> RightBoundsIncrease;
+	TSharedPtr<FUICommandInfo> RightBoundsDecrease;
+	TSharedPtr<FUICommandInfo> BoundsIncrementIncrease;
+	TSharedPtr<FUICommandInfo> BoundsIncrementDecrease;
+	TSharedPtr<FUICommandInfo> SelectNextLoop;
+	TSharedPtr<FUICommandInfo> SelectPreviousLoop;
 };
diff --git a/Engine/Plugins/Editor/WaveformEditor/Source/WaveformEditorWidgets/Private/SWaveformTransformationRenderLayer.cpp b/Engine/Plugins/Editor/WaveformEditor/Source/WaveformEditorWidgets/Private/SWaveformTransformationRenderLayer.cpp
--- a/Engine/Plugins/Editor/WaveformEditor/Source/WaveformEditorWidgets/Private/SWaveformTransformationRenderLayer.cpp
+++ b/Engine/Plugins/Editor/WaveformEditor/Source/WaveformEditorWidgets/Private/SWaveformTransformationRenderLayer.cpp
@@ -2,6 +2,7 @@
 
 #include "SWaveformTransformationRenderLayer.h"
 
+#include "Framework/Application/SlateApplication.h"
 #include "Styling/AppStyle.h"
 
 void SWaveformTransformationRenderLayer::PrivateRegisterAttributes(FSlateAttributeInitializer& AttributeInitializer)
@@ -40,6 +41,7 @@ FReply SWaveformTransformationRenderLayer::OnMouseButtonDown(const FGeometry& My
 
 FReply SWaveformTransformationRenderLayer::OnMouseButtonUp(const FGeometry& MyGeometry, const FPointerEvent& MouseEvent)
 {
+	FSlateApplication::Get().SetKeyboardFocus(SharedThis(this));
 	return TransformationRenderer->OnMouseButtonUp(*this, MyGeometry, MouseEvent);
 }
 
@@ -68,3 +70,8 @@ void SWaveformTransformationRenderLayer::OnMouseLeave(const FPointerEvent& Mouse
 	return TransformationRenderer->OnMouseLeave(MouseEvent);
 }
 
+FReply SWaveformTransformationRenderLayer::OnKeyUp(const FGeometry& MyGeometry, const FKeyEvent& InKeyEvent)
+{
+	return TransformationRenderer->OnKeyUp(MyGeometry, InKeyEvent);
+}
+
diff --git a/Engine/Plugins/Editor/WaveformEditor/Source/WaveformEditorWidgets/Public/SWaveformTransformationRenderLayer.h b/Engine/Plugins/Editor/WaveformEditor/Source/WaveformEditorWidgets/Public/SWaveformTransformationRenderLayer.h
--- a/Engine/Plugins/Editor/WaveformEditor/Source/WaveformEditorWidgets/Public/SWaveformTransformationRenderLayer.h
+++ b/Engine/Plugins/Editor/WaveformEditor/Source/WaveformEditorWidgets/Public/SWaveformTransformationRenderLayer.h
@@ -27,6 +27,8 @@ public:
 	FReply OnMouseButtonDoubleClick(const FGeometry& InMyGeometry, const FPointerEvent& InMouseEvent) override;
 	void OnMouseEnter(const FGeometry& MyGeometry, const FPointerEvent& MouseEvent) override;
 	void OnMouseLeave(const FPointerEvent& MouseEvent) override;
+	FReply OnKeyUp(const FGeometry& MyGeometry, const FKeyEvent& InKeyEvent) override;
+	virtual bool SupportsKeyboardFocus() const override { return true; }
 
 private:
 	TSharedPtr<IWaveformTransformationRenderer> TransformationRenderer;
diff --git a/Engine/Plugins/Editor/WaveformEditor/Source/WaveformTransformations/Private/WaveformTransformationMarkers.cpp b/Engine/Plugins/Editor/WaveformEditor/Source/WaveformTransformations/Private/WaveformTransformationMarkers.cpp
--- a/Engine/Plugins/Editor/WaveformEditor/Source/WaveformTransformations/Private/WaveformTransformationMarkers.cpp
+++ b/Engine/Plugins/Editor/WaveformEditor/Source/WaveformTransformations/Private/WaveformTransformationMarkers.cpp
@@ -28,6 +28,18 @@ void UWaveformTransformationMarkers::OverwriteTransformation()
 	Markers->Reset();
 }
 
+void UWaveformTransformationMarkers::ModifyMarkerLoopRegion(ELoopModificationControls Modification) const
+{
+	check(Markers != nullptr);
+	Markers->ModifyMarkerLoop.ExecuteIfBound(Modification);
+}
+
+void UWaveformTransformationMarkers::CycleMarkerLoopRegion(ELoopModificationControls Modification) const
+{
+	check(Markers != nullptr);
+	Markers->CycleMarkerLoop.ExecuteIfBound(Modification);
+}
+
 #if WITH_EDITOR
 void UWaveformTransformationMarkers::OverwriteSoundWaveData(USoundWave& InOutSoundWave)
 {
diff --git a/Engine/Plugins/Editor/WaveformEditor/Source/WaveformTransformations/Public/WaveformTransformationMarkers.h b/Engine/Plugins/Editor/WaveformEditor/Source/WaveformTransformations/Public/WaveformTransformationMarkers.h
--- a/Engine/Plugins/Editor/WaveformEditor/Source/WaveformTransformations/Public/WaveformTransformationMarkers.h
+++ b/Engine/Plugins/Editor/WaveformEditor/Source/WaveformTransformations/Public/WaveformTransformationMarkers.h
@@ -8,6 +8,19 @@
 
 class USoundWave;
 
+enum class ELoopModificationControls : uint8
+{
+	None = 0,
+	LeftHandleIncrement,
+	LeftHandleDecrement,
+	RightHandleIncrement,
+	RightHandleDecrement,
+	IncreaseIncrement,
+	DecreaseIncrement,
+	SelectNextLoop,
+	SelectPreviousLoop,
+};
+
 //Used to make CuePoints work with PropertyHandles
 UCLASS(BlueprintType, DefaultToInstanced)
 class WAVEFORMTRANSFORMATIONS_API UWaveCueArray : public UObject
@@ -26,6 +39,12 @@ public:
 	UPROPERTY(EditAnywhere, Category = "Markers")
 	TArray<FSoundWaveCuePoint> WaveCues;
 
+	DECLARE_DELEGATE_OneParam(ModifyMarkerLoopRegionDelegate, ELoopModificationControls);
+	ModifyMarkerLoopRegionDelegate ModifyMarkerLoop;
+
+	DECLARE_DELEGATE_OneParam(CycleMarkerLoopRegionDelegate, ELoopModificationControls);
+	CycleMarkerLoopRegionDelegate CycleMarkerLoop;
+
 private:
 	UPROPERTY()
 	bool bIsInitialized = false;
@@ -52,6 +71,9 @@ public:
 	void UpdateConfiguration(FWaveTransformUObjectConfiguration& InOutConfiguration) override;
 	void OverwriteTransformation() override;
 
+	void ModifyMarkerLoopRegion(ELoopModificationControls Modification) const;
+	void CycleMarkerLoopRegion(ELoopModificationControls Modification) const;
+
 #if WITH_EDITOR
 	void OverwriteSoundWaveData(USoundWave& InOutSoundWave) override;
 	void GetTransformationInfo(FWaveformTransformationInfo& InOutTransformationInfo) const override;
diff --git a/Engine/Plugins/Editor/WaveformEditor/Source/WaveformTransformationsWidgets/Private/WaveformTransformationMarkersRenderer.cpp b/Engine/Plugins/Editor/WaveformEditor/Source/WaveformTransformationsWidgets/Private/WaveformTransformationMarkersRenderer.cpp
--- a/Engine/Plugins/Editor/WaveformEditor/Source/WaveformTransformationsWidgets/Private/WaveformTransformationMarkersRenderer.cpp
+++ b/Engine/Plugins/Editor/WaveformEditor/Source/WaveformTransformationsWidgets/Private/WaveformTransformationMarkersRenderer.cpp
@@ -218,6 +218,167 @@ void FWaveformTransformationMarkerRenderer::SetPropertyHandles(const TArray<TSha
 	FWaveformTransformationRendererBase::SetPropertyHandles(InPropertyHandles);
 
 	MarkersHandle = GetPropertyHandle(GET_MEMBER_NAME_CHECKED(UWaveformTransformationMarkers, Markers));
+
+	MarkersArray = CastChecked<UWaveCueArray>(GetPropertyValue<UObject*>(MarkersHandle));
+
+	MarkersArray->ModifyMarkerLoop.BindLambda([this](ELoopModificationControls Modification) { ModifyMarkerLoopRegion(Modification); });
+	MarkersArray->CycleMarkerLoop.BindLambda([this](ELoopModificationControls Modification) { CycleMarkerLoopRegion(Modification); });
+}
+
+void FWaveformTransformationMarkerRenderer::ModifyMarkerLoopRegion(ELoopModificationControls Modification)
+{
+	FSoundWaveCuePoint* LoopRegion = GetSelectedMarker();
+
+	if (LoopRegion == nullptr)
+	{
+		return;
+	}
+
+	check(MarkersArray);
+	check(TransformationWaveInfo.NumChannels > 0);
+	check(TransformationWaveInfo.NumAvilableSamples > 0);
+	check(SizeIncrements.Num() > 0);
+
+	int32 FramesToShift = (SizeIncrements[SelectedIncrement] / 1000.0f) * TransformationWaveInfo.SampleRate;
+	int32 MaxFrames = TransformationWaveInfo.NumAvilableSamples;
+
+	MarkersArray->Modify();
+
+	switch (Modification)
+	{
+	case ELoopModificationControls::None:
+		break;
+	case ELoopModificationControls::LeftHandleIncrement:
+	{
+		LoopRegion->FramePosition -= FramesToShift;
+		LoopRegion->FrameLength += FramesToShift;
+	}
+		break;
+	case ELoopModificationControls::LeftHandleDecrement:
+	{
+		LoopRegion->FramePosition += FramesToShift;
+		LoopRegion->FrameLength -= FramesToShift;
+	}
+		break;
+	case ELoopModificationControls::RightHandleIncrement:
+	{
+		LoopRegion->FrameLength += FramesToShift;
+	}
+		break;
+	case ELoopModificationControls::RightHandleDecrement:
+	{
+		LoopRegion->FrameLength -= FramesToShift;
+	}
+		break;
+	case ELoopModificationControls::IncreaseIncrement:
+	{
+		if (SelectedIncrement < SizeIncrements.Num() - 1)
+		{
+			++SelectedIncrement;
+		}
+	}
+	break;
+	case ELoopModificationControls::DecreaseIncrement:
+	{
+		if (SelectedIncrement > 0)
+		{
+			--SelectedIncrement;
+		}
+	}
+	default:
+		break;
+	}
+
+	LoopRegion->FramePosition = FMath::Clamp(LoopRegion->FramePosition, 0, LoopRegion->FramePosition);
+	LoopRegion->FrameLength = FMath::Clamp(LoopRegion->FrameLength, 0, MaxFrames - LoopRegion->FramePosition);
+
+	EndTransaction();
+}
+
+void FWaveformTransformationMarkerRenderer::CycleMarkerLoopRegion(ELoopModificationControls Modification)
+{
+	check(MarkersArray);
+
+	const int32 MarkerArrayLength = MarkersArray->WaveCues.Num();
+
+	// If there is no active selection, pick first or last elmeent
+	if (SelectedCue == INDEX_NONE && MarkerArrayLength > 0)
+	{
+		switch (Modification)
+		{
+		case ELoopModificationControls::SelectNextLoop:
+		{
+			SelectedCue = MarkersArray->WaveCues[0].CuePointID;
+		}
+			break;
+		case ELoopModificationControls::SelectPreviousLoop:
+		{
+			SelectedCue = MarkersArray->WaveCues[MarkerArrayLength - 1].CuePointID;
+		}
+			break;
+		default:
+			break;
+		}
+	}
+
+	// Since cues can be removed and added the ids do not match list indexes, so search for matches
+	for (int i = 0; i < MarkerArrayLength; i++)
+	{
+		if (MarkersArray->WaveCues[i].CuePointID == SelectedCue)
+		{
+			// Check if we move backwards of forwards
+			if (Modification == ELoopModificationControls::SelectNextLoop && i < MarkerArrayLength - 1)
+			{
+				SelectedCue = MarkersArray->WaveCues[i + 1].CuePointID;
+				break;
+			}
+			else if (Modification == ELoopModificationControls::SelectPreviousLoop && i > 0)
+			{
+				SelectedCue = MarkersArray->WaveCues[i - 1].CuePointID;
+				break;
+			}
+		}
+	}
+
+}
+
+FSoundWaveCuePoint* FWaveformTransformationMarkerRenderer::GetSelectedMarker() const
+{
+	check(MarkersArray);
+
+	if (SelectedCue == INDEX_NONE)
+	{
+		return nullptr;
+	}
+
+	for (int i = 0; i < MarkersArray->WaveCues.Num(); i++)
+	{
+		if (MarkersArray->WaveCues[i].CuePointID == SelectedCue)
+		{
+			return &MarkersArray->WaveCues[i];
+		}
+	}
+
+	return nullptr;
+}
+
+void FWaveformTransformationMarkerRenderer::SelectLoopRegionByKeyboard(const FKey& PressedKey)
+{
+	check(MarkersArray);
+
+	for (int i = 0; i < IndexKeyCodes.Num(); i++)
+	{
+		// Check we have loop regions in this index
+		if (MarkersArray->WaveCues.Num() <= i)
+		{
+			return;
+		}
+
+		if (PressedKey == IndexKeyCodes[i])
+		{
+			SelectedCue = i;
+		}
+	}
 }
 
 void FWaveformTransformationMarkerRenderer::Tick(const FGeometry& AllottedGeometry, const double InCurrentTime, const float InDeltaTime)
@@ -379,6 +540,12 @@ FCursorReply FWaveformTransformationMarkerRenderer::OnCursorQuery(const FGeometr
 	return FCursorReply::Unhandled();
 }
 
+FReply FWaveformTransformationMarkerRenderer::OnKeyUp(const FGeometry& MyGeometry, const FKeyEvent& InKeyEvent)
+{
+	SelectLoopRegionByKeyboard(InKeyEvent.GetKey());
+	return FReply::Handled();
+}
+
 void FWaveformTransformationMarkerRenderer::OnWidgetSettingsUpdated(const FName& PropertyName, const UWaveformTransformationsWidgetsSettings* Settings)
 {
 	check(Settings)
diff --git a/Engine/Plugins/Editor/WaveformEditor/Source/WaveformTransformationsWidgets/Private/WaveformTransformationMarkersRenderer.h b/Engine/Plugins/Editor/WaveformEditor/Source/WaveformTransformationsWidgets/Private/WaveformTransformationMarkersRenderer.h
--- a/Engine/Plugins/Editor/WaveformEditor/Source/WaveformTransformationsWidgets/Private/WaveformTransformationMarkersRenderer.h
+++ b/Engine/Plugins/Editor/WaveformEditor/Source/WaveformTransformationsWidgets/Private/WaveformTransformationMarkersRenderer.h
@@ -19,9 +19,15 @@ public:
 	virtual FReply OnMouseButtonUp(SWidget& OwnerWidget, const FGeometry& MyGeometry, const FPointerEvent& MouseEvent) override;
 	virtual FReply OnMouseMove(SWidget& OwnerWidget, const FGeometry& MyGeometry, const FPointerEvent& MouseEvent) override;
 	virtual FCursorReply OnCursorQuery(const FGeometry& MyGeometry, const FPointerEvent& CursorEvent) const override;
+	virtual FReply OnKeyUp(const FGeometry& MyGeometry, const FKeyEvent& InKeyEvent) override;
 
 	virtual void SetPropertyHandles(const TArray<TSharedRef<IPropertyHandle>>& InPropertyHandles) override;
 
+	void ModifyMarkerLoopRegion(ELoopModificationControls Modification);
+	void CycleMarkerLoopRegion(ELoopModificationControls Modification);
+	FSoundWaveCuePoint* GetSelectedMarker() const;
+	void SelectLoopRegionByKeyboard(const FKey& PressedKey);
+
 private:
 	int32 DrawMarkerHandles(const FGeometry& AllottedGeometry, FSlateWindowElementList& OutDrawElements, int32 LayerId) const;
 
@@ -92,5 +98,9 @@ private:
 	float LabelFontSize = WaveformTransformationWidgetSharedDefaults::DefaultLabelFontSize;
 
 	TSharedPtr<IPropertyHandle> MarkersHandle;
+
 	TWeakPtr<SWindow> CreateLoopMenuWindow;
+	const TArray<int32> SizeIncrements = { 1, 10, 50, 100 };
+	const TArray<FKey> IndexKeyCodes = { EKeys::One, EKeys::Two, EKeys::Three, EKeys::Four, EKeys::Five, EKeys::Six, EKeys::Seven, EKeys::Eight, EKeys::Nine };
+	int32 SelectedIncrement = 0;
 };
diff --git a/Engine/Plugins/Editor/WaveformEditor/Source/WaveformTransformationsWidgets/Private/WaveformTransformationRendererBase.cpp b/Engine/Plugins/Editor/WaveformEditor/Source/WaveformTransformationsWidgets/Private/WaveformTransformationRendererBase.cpp
--- a/Engine/Plugins/Editor/WaveformEditor/Source/WaveformTransformationsWidgets/Private/WaveformTransformationRendererBase.cpp
+++ b/Engine/Plugins/Editor/WaveformEditor/Source/WaveformTransformationsWidgets/Private/WaveformTransformationRendererBase.cpp
@@ -43,6 +43,11 @@ FCursorReply FWaveformTransformationRendererBase::OnCursorQuery(const FGeometry&
 	return FCursorReply::Unhandled();
 }
 
+FReply FWaveformTransformationRendererBase::OnKeyUp(const FGeometry& MyGeometry, const FKeyEvent& InKeyEvent)
+{
+	return FReply::Unhandled();
+}
+
 void FWaveformTransformationRendererBase::SetTransformationWaveInfo(const FWaveformTransformationRenderInfo& InWaveInfo)
 {
 	TransformationWaveInfo = InWaveInfo;
diff --git a/Engine/Plugins/Editor/WaveformEditor/Source/WaveformTransformationsWidgets/Public/IWaveformTransformationRenderer.h b/Engine/Plugins/Editor/WaveformEditor/Source/WaveformTransformationsWidgets/Public/IWaveformTransformationRenderer.h
--- a/Engine/Plugins/Editor/WaveformEditor/Source/WaveformTransformationsWidgets/Public/IWaveformTransformationRenderer.h
+++ b/Engine/Plugins/Editor/WaveformEditor/Source/WaveformTransformationsWidgets/Public/IWaveformTransformationRenderer.h
@@ -12,6 +12,7 @@ class FReply;
 class FSlateRect;
 class FSlateWindowElementList;
 class FWidgetStyle;
+struct FKeyEvent;
 class SWidget;
 class UClass;
 struct FGeometry;
@@ -41,6 +42,7 @@ public:
 	virtual void OnMouseEnter(const FGeometry& MyGeometry, const FPointerEvent& MouseEvent) = 0;
 	virtual void OnMouseLeave(const FPointerEvent& MouseEvent) = 0;
 	virtual FCursorReply OnCursorQuery(const FGeometry& MyGeometry, const FPointerEvent& CursorEvent) const = 0;
+	virtual FReply OnKeyUp(const FGeometry& MyGeometry, const FKeyEvent& InKeyEvent) = 0;
 
 	/* Gets info about the transformation being rendered */
 	virtual void SetTransformationWaveInfo(const FWaveformTransformationRenderInfo& InWaveInfo) = 0;
diff --git a/Engine/Plugins/Editor/WaveformEditor/Source/WaveformTransformationsWidgets/Public/WaveformTransformationRendererBase.h b/Engine/Plugins/Editor/WaveformEditor/Source/WaveformTransformationsWidgets/Public/WaveformTransformationRendererBase.h
--- a/Engine/Plugins/Editor/WaveformEditor/Source/WaveformTransformationsWidgets/Public/WaveformTransformationRendererBase.h
+++ b/Engine/Plugins/Editor/WaveformEditor/Source/WaveformTransformationsWidgets/Public/WaveformTransformationRendererBase.h
@@ -21,6 +21,7 @@ public:
 	virtual void OnMouseEnter(const FGeometry& MyGeometry, const FPointerEvent& MouseEvent) override {};
 	virtual void OnMouseLeave(const FPointerEvent& MouseEvent) override {};
 	virtual FCursorReply OnCursorQuery(const FGeometry& MyGeometry, const FPointerEvent& CursorEvent) const override;
+	virtual FReply OnKeyUp(const FGeometry& MyGeometry, const FKeyEvent& InKeyEvent) override;
 
 	virtual void SetTransformationWaveInfo(const FWaveformTransformationRenderInfo& InWaveInfo) override;
 	virtual void SetPropertyHandles(const TArray<TSharedRef<IPropertyHandle>>& InPropertyHandles) override;
