Waveform Editor: Removed edit waveform context menu option

diff --git a/Engine/Plugins/Editor/WaveformEditor/Source/WaveformEditor/Private/WaveformEditorInstantiator.cpp b/Engine/Plugins/Editor/WaveformEditor/Source/WaveformEditor/Private/WaveformEditorInstantiator.cpp
--- a/Engine/Plugins/Editor/WaveformEditor/Source/WaveformEditor/Private/WaveformEditorInstantiator.cpp
+++ b/Engine/Plugins/Editor/WaveformEditor/Source/WaveformEditor/Private/WaveformEditorInstantiator.cpp
@@ -28,22 +28,6 @@ void FWaveformEditorInstantiator::ExtendContentBrowserSelectionMenu()
 {
 	UToolMenu* Menu = UToolMenus::Get()->ExtendMenu("ContentBrowser.AssetContextMenu.SoundWave");
 	FToolMenuSection& Section = Menu->FindOrAddSection("GetAssetActions");
-
-	Section.AddDynamicEntry("SoundWaveEditing_CreateWaveformEditor", FNewToolMenuSectionDelegate::CreateLambda([this](FToolMenuSection& InSection)
-		{
-			if (const UContentBrowserAssetContextMenuContext* Context = InSection.FindContext<UContentBrowserAssetContextMenuContext>())
-			{
-				if (Algo::AnyOf(Context->SelectedAssets, WaveEditorInstantiatorPrivate::FilterUnwantedAssets))
-				{
-					const TAttribute<FText> Label = LOCTEXT("SoundWave_WaveformEditor", "Edit Waveform");
-					const TAttribute<FText> ToolTip = LOCTEXT("SoundWave_WaveformEditor_Tooltip", "Open waveform editor");
-					const FSlateIcon Icon = FSlateIcon();
-					const FToolMenuExecuteAction UIAction = FToolMenuExecuteAction::CreateSP(this, &FWaveformEditorInstantiator::ExecuteCreateWaveformEditor);
-
-					InSection.AddMenuEntry("SoundWave_CreateWaveformEditor", Label, ToolTip, Icon, UIAction);
-				}
-			}
-		}));
 	
 	FToolMenuSection& ImportSection = Menu->FindOrAddSection("ImportedAssetActions");
 
@@ -70,15 +54,6 @@ void FWaveformEditorInstantiator::RegisterAsSoundwaveEditor()
 	AudioEditorModule.SoundWaveEditorOpen.BindSP(this, &FWaveformEditorInstantiator::CreateWaveformEditor);
 }
 
-void FWaveformEditorInstantiator::ExecuteCreateWaveformEditor(const FToolMenuContext& MenuContext)
-{
-	if (const UContentBrowserAssetContextMenuContext* Context = UContentBrowserAssetContextMenuContext::FindContextWithAssets(MenuContext))
-	{
-		TArray<USoundWave*> SoundWavesToEdit = Context->LoadSelectedObjectsIf<USoundWave>(WaveEditorInstantiatorPrivate::FilterUnwantedAssets);
-		CreateWaveformEditor(SoundWavesToEdit);
-	}
-}
-
 void FWaveformEditorInstantiator::ExecuteReImportAndOverwriteTransformations(const FToolMenuContext& MenuContext)
 {
 	if (const UContentBrowserAssetContextMenuContext* Context = UContentBrowserAssetContextMenuContext::FindContextWithAssets(MenuContext))
diff --git a/Engine/Plugins/Editor/WaveformEditor/Source/WaveformEditor/Private/WaveformEditorInstantiator.h b/Engine/Plugins/Editor/WaveformEditor/Source/WaveformEditor/Private/WaveformEditorInstantiator.h
--- a/Engine/Plugins/Editor/WaveformEditor/Source/WaveformEditor/Private/WaveformEditorInstantiator.h
+++ b/Engine/Plugins/Editor/WaveformEditor/Source/WaveformEditor/Private/WaveformEditorInstantiator.h
@@ -19,7 +19,6 @@ public:
 
 private:
 	virtual void CreateWaveformEditor(TArray<USoundWave*> SoundWavesToEdit) override;
-	void ExecuteCreateWaveformEditor(const FToolMenuContext& MenuContext);
 	void ExecuteReImportAndOverwriteTransformations(const FToolMenuContext& MenuContext);
 	bool CanSoundWaveBeOpenedInEditor(const USoundWave* SoundWaveToEdit);
 	void DisplayErrorDialog(const FText& ErrorMessage) const;
