Waveform Editor: Fixes for undo actions

diff --git a/Engine/Plugins/Editor/WaveformEditor/Source/WaveformTransformations/Private/WaveformTransformationMarkers.cpp b/Engine/Plugins/Editor/WaveformEditor/Source/WaveformTransformations/Private/WaveformTransformationMarkers.cpp
--- a/Engine/Plugins/Editor/WaveformEditor/Source/WaveformTransformations/Private/WaveformTransformationMarkers.cpp
+++ b/Engine/Plugins/Editor/WaveformEditor/Source/WaveformTransformations/Private/WaveformTransformationMarkers.cpp
@@ -88,12 +88,10 @@ void UWaveCueArray::PostEditChangeProperty(FPropertyChangedEvent& PropertyChange
 	Super::PostEditChangeProperty(PropertyChangedEvent);
 
 	const FName PropertyName = PropertyChangedEvent.GetPropertyName();
-	if (PropertyName == GET_MEMBER_NAME_CHECKED(UWaveCueArray, WaveCues) &&
+	if (PropertyName == GET_MEMBER_NAME_CHECKED(UWaveCueArray, WaveCues) && WaveCues.Num() > 0 &&
 		(PropertyChangedEvent.ChangeType & (EPropertyChangeType::ArrayAdd | EPropertyChangeType::ResetToDefault | 
 			EPropertyChangeType::Unspecified | EPropertyChangeType::Redirected)))
 	{
-		check(WaveCues.Num() > 0);
-
 		int32 LargestCuePointID = 1;
 		for (int32 Index = WaveCues.Num()-1; Index >= 0; Index--)
 		{
diff --git a/Engine/Plugins/Editor/WaveformEditor/Source/WaveformTransformationsWidgets/Private/WaveformTransformationMarkersRenderer.cpp b/Engine/Plugins/Editor/WaveformEditor/Source/WaveformTransformationsWidgets/Private/WaveformTransformationMarkersRenderer.cpp
--- a/Engine/Plugins/Editor/WaveformEditor/Source/WaveformTransformationsWidgets/Private/WaveformTransformationMarkersRenderer.cpp
+++ b/Engine/Plugins/Editor/WaveformEditor/Source/WaveformTransformationsWidgets/Private/WaveformTransformationMarkersRenderer.cpp
@@ -393,8 +393,8 @@ void FWaveformTransformationMarkerRenderer::Tick(const FGeometry& AllottedGeomet
 	FVector2D MouseAbsolutePosition = UWidgetLayoutLibrary::GetMousePositionOnPlatform();
 	MousePosition = AllottedGeometry.AbsoluteToLocal(MouseAbsolutePosition);
 
-	MarkersArray = CastChecked<UWaveCueArray>(GetPropertyValue<UObject*>(MarkersHandle));
-	
+	check(MarkersArray != nullptr);
+
 	TMap<uint32, FSoundWaveCuePoint> CuePointMap;
 	CuePointMap.Reserve(MarkersArray->WaveCues.Num());
 
@@ -465,6 +465,11 @@ FReply FWaveformTransformationMarkerRenderer::OnMouseButtonUp(SWidget& OwnerWidg
 		FSlateApplication::Get().RequestDestroyWindow(LockedFadeModeMenuWindow.ToSharedRef());
 	}
 
+	if (MarkersArray == nullptr)
+	{
+		return FReply::Unhandled();
+	}
+
 	if (MarkerInteractionType != EMarkerInteractionType::None)
 	{
 		MarkersArray->Modify();
@@ -489,6 +494,11 @@ FReply FWaveformTransformationMarkerRenderer::OnMouseMove(SWidget& OwnerWidget,
 		return FReply::Handled().CaptureMouse(OwnerWidget.AsShared());
 	}
 
+	if (MarkersArray == nullptr)
+	{
+		return FReply::Unhandled();
+	}
+
 	// Evaluate highlight state for loop handles
 	HighlightedCue = INDEX_NONE;
 	for (const FSoundWaveCuePoint& CuePoint : MarkersArray->WaveCues)
@@ -510,6 +520,11 @@ FCursorReply FWaveformTransformationMarkerRenderer::OnCursorQuery(const FGeometr
 		return FCursorReply::Cursor(EMouseCursor::GrabHandClosed);
 	}
 
+	if (MarkersArray == nullptr)
+	{
+		return FCursorReply::Unhandled();
+	}
+
 	// Set resize on loop handle hover, loops are seperated as we want highlight interactions to superscede region highlights
 	for (const FSoundWaveCuePoint& CuePoint : MarkersArray->WaveCues)
 	{
@@ -652,6 +667,7 @@ bool FWaveformTransformationMarkerRenderer::SetMarkerInInteractionRange(const FV
 {
 	check(TransformationWaveInfo.NumChannels > 0);
 	ensure(LoopHandlePixelWidth > 0);
+	check(MarkersArray != nullptr);
 
 	InteractionRanges.Reset(MarkersArray->WaveCues.Num());
 	
@@ -726,6 +742,8 @@ bool FWaveformTransformationMarkerRenderer::SelectLoopArea(const FVector2D& InLo
 
 void FWaveformTransformationMarkerRenderer::SetActiveLoopRegion(int32 CuePointID)
 {
+	check(MarkersArray != nullptr);
+
 	MarkersHandle->NotifyPreChange();
 	MarkersArray->Modify();
 	SelectedCue = CuePointID;
@@ -755,6 +773,11 @@ FWaveformTransformationMarkerRenderer::EMarkerInteractionType FWaveformTransform
 		return EMarkerInteractionType::None;
 	}
 
+	if (MarkersArray == nullptr)
+	{
+		return EMarkerInteractionType::None;
+	}
+
 	for (FSoundWaveCuePoint CuePoint : MarkersArray->WaveCues)
 	{
 		if (!CuePoint.IsLoopRegion())
@@ -846,6 +869,7 @@ void FWaveformTransformationMarkerRenderer::SetPropertyValueDependingOnInteracti
 void FWaveformTransformationMarkerRenderer::ShowConvertToLoopAtMarker(const FGeometry& MyGeometry, const FPointerEvent& MouseEvent, FSoundWaveCuePoint* CueToModify)
 {
 	check(CueToModify);
+	check(MarkersArray != nullptr);
 	const FVector2D LocalCursorPosition = MouseEvent.GetScreenSpacePosition();
 
 	TSharedRef<SVerticalBox> MenuContent = SNew(SVerticalBox);
@@ -862,7 +886,6 @@ void FWaveformTransformationMarkerRenderer::ShowConvertToLoopAtMarker(const FGeo
 				.Text(LOCTEXT("WaveformTransformationMarkerRendererConvertLoop","Convert To Loop Region"))
 				.OnClicked_Lambda([this, CueToModify]() -> FReply
 					{
-						//ToggleLoopRegion(CueToModify);
 						MarkersArray->EnableLoopRegion(CueToModify);
 						CueToModify->FrameLength = MarkerInInteractionRange.GetMinLoopSize();
 
