Waveform Editor: UI to add or remove markers and loop regions

diff --git a/Engine/Plugins/Editor/WaveformEditor/Source/WaveformEditor/Private/WaveformEditor.cpp b/Engine/Plugins/Editor/WaveformEditor/Source/WaveformEditor/Private/WaveformEditor.cpp
--- a/Engine/Plugins/Editor/WaveformEditor/Source/WaveformEditor/Private/WaveformEditor.cpp
+++ b/Engine/Plugins/Editor/WaveformEditor/Source/WaveformEditor/Private/WaveformEditor.cpp
@@ -209,6 +209,42 @@ TObjectPtr<UWaveformTransformationTrimFade> FWaveformEditor::GetOrAddTrimFadeTra
 	return TrimFadeTransformation;
 }
 
+TObjectPtr<UWaveformTransformationMarkers> FWaveformEditor::GetOrAddMarkerTransformation()
+{
+	check(SoundWave);
+
+	TObjectPtr<UWaveformTransformationMarkers> MarkersTransformation;
+
+	for (TObjectPtr<UWaveformTransformationBase> Transformation : SoundWave->Transformations)
+	{
+		MarkersTransformation = Cast<UWaveformTransformationMarkers>(Transformation);
+		if (MarkersTransformation)
+		{
+			break;
+		}
+	}
+
+	if (MarkersTransformation == nullptr)
+	{
+		EObjectFlags MaskedOuterFlags = SoundWave ? SoundWave->GetMaskedFlags(RF_PropagateToSubObjects) : RF_NoFlags;
+
+		if (SoundWave->HasAnyFlags(RF_ClassDefaultObject | RF_ArchetypeObject))
+		{
+			MaskedOuterFlags |= RF_ArchetypeObject;
+		}
+
+		MarkersTransformation = NewObject<UWaveformTransformationMarkers>(SoundWave, UWaveformTransformationMarkers::StaticClass(), NAME_None, MaskedOuterFlags);
+
+		SoundWave->Modify();
+		SoundWave->Transformations.Add(MarkersTransformation);
+
+		NotifyPostTransformationChange();
+	}
+
+	return MarkersTransformation;
+}
+
+
 void FWaveformEditor::ToggleFadeIn()
 {
 	check(SoundWave);
@@ -287,39 +323,58 @@ bool FWaveformEditor::CanFadeOut()
 	return true;
 }
 
-TObjectPtr<UWaveformTransformationMarkers> FWaveformEditor::GetOrAddMarkerTransformation()
+void FWaveformEditor::CreateMarker(bool bIsLoopRegion)
 {
-	check(SoundWave);
+	TObjectPtr<UWaveformTransformationMarkers> MarkerTransformation = GetOrAddMarkerTransformation();
+	check(MarkerTransformation);
 
-	TObjectPtr<UWaveformTransformationMarkers> MarkersTransformation;
+	int32 HighestCueID = INDEX_NONE;
 
-	for (TObjectPtr<UWaveformTransformationBase> Transformation : SoundWave->Transformations)
+	SoundWave->Modify();
+
+	for (FSoundWaveCuePoint CuePoint : MarkerTransformation->Markers->WaveCues)
 	{
-		MarkersTransformation = Cast<UWaveformTransformationMarkers>(Transformation);
-		if (MarkersTransformation)
+		if (CuePoint.CuePointID > HighestCueID)
 		{
-			break;
+			HighestCueID = CuePoint.CuePointID;
 		}
 	}
 
-	if (MarkersTransformation == nullptr)
+	FSoundWaveCuePoint NewCuePoint;
+	NewCuePoint.CuePointID = HighestCueID + 1;
+	NewCuePoint.FramePosition = TransportCoordinator->GetFocusPoint() * SoundWave->TotalSamples;
+
+	if (bIsLoopRegion)
 	{
-		EObjectFlags MaskedOuterFlags = SoundWave ? SoundWave->GetMaskedFlags(RF_PropagateToSubObjects) : RF_NoFlags;
+		NewCuePoint.SetLoopRegion(true);
+		NewCuePoint.FrameLength = SoundWave->TotalSamples * 0.1; // Default loop region to 10% of samples for easier tuning
+	}
 
-		if (SoundWave->HasAnyFlags(RF_ClassDefaultObject | RF_ArchetypeObject))
-		{
-			MaskedOuterFlags |= RF_ArchetypeObject;
-		}
+	MarkerTransformation->Markers->WaveCues.Add(NewCuePoint);
+	NotifyPostTransformationChange(EPropertyChangeType::ValueSet);
+}
 
-		MarkersTransformation = NewObject<UWaveformTransformationMarkers>(SoundWave, UWaveformTransformationMarkers::StaticClass(), NAME_None, MaskedOuterFlags);
+void FWaveformEditor::DeleteMarker()
+{
+	TObjectPtr<UWaveformTransformationMarkers> MarkerTransformation = GetOrAddMarkerTransformation();
+	check(MarkerTransformation);
 
-		SoundWave->Modify();
-		SoundWave->Transformations.Add(MarkersTransformation);
+	int32 CueToDelete = MarkerTransformation->Markers->SelectedCue;
 
-		NotifyPostTransformationChange();
+	if (CueToDelete == INDEX_NONE)
+	{
+		return; // No cue to selected
 	}
 
-	return MarkersTransformation;
+	for (int i = 0; i < MarkerTransformation->Markers->WaveCues.Num(); i++)
+	{
+		if (CueToDelete == MarkerTransformation->Markers->WaveCues[i].CuePointID)
+		{
+			MarkerTransformation->Markers->WaveCues.RemoveAt(i);
+			MarkerTransformation->Markers->SelectedCue = INDEX_NONE;
+			break;
+		}
+	}
 }
 
 void FWaveformEditor::ModifyMarkerLoopRegion(ELoopModificationControls Modification)
@@ -708,6 +763,55 @@ bool FWaveformEditor::RegisterToolbar()
 					InSection.AddEntry(ExportOptionsEntry);
 				}
 			}));
+
+		TransformationSection.AddDynamicEntry("CreateMarkerButton", FNewToolMenuSectionDelegate::CreateLambda([this, Commands](FToolMenuSection& InSection)
+			{
+				const UWaveformEditorToolMenuContext* Context = InSection.FindContext<UWaveformEditorToolMenuContext>();
+
+				if (Context && Context->WaveformEditor.IsValid())
+				{
+					FToolMenuEntry CreateMarker = FToolMenuEntry::InitToolBarButton(
+						Commands.CreateMarker,
+						FText(),
+						LOCTEXT("WaveformEditorCreateMarkerButtonTooltip", "Create a marker cue for the soundwave"),
+						FSlateIcon(FAppStyle::GetAppStyleSetName(), "Sequencer.Tracks.Event")
+					);
+					InSection.AddEntry(CreateMarker);
+				}
+			}));
+
+
+		TransformationSection.AddDynamicEntry("CreateLoopRegionButton", FNewToolMenuSectionDelegate::CreateLambda([this, Commands](FToolMenuSection& InSection)
+			{
+				const UWaveformEditorToolMenuContext* Context = InSection.FindContext<UWaveformEditorToolMenuContext>();
+
+				if (Context && Context->WaveformEditor.IsValid())
+				{
+					FToolMenuEntry CreateLoopRegion = FToolMenuEntry::InitToolBarButton(
+						Commands.CreateLoopRegion,
+						FText(),
+						LOCTEXT("WaveformEditorCreateLoopRegionButtonTooltip", "Create a marker loop region for the soundwave"),
+						FSlateIcon(FAppStyle::GetAppStyleSetName(), "Graph.Node.Loop")
+					);
+					InSection.AddEntry(CreateLoopRegion);
+				}
+			}));
+
+		TransformationSection.AddDynamicEntry("DeleteMarkerButton", FNewToolMenuSectionDelegate::CreateLambda([this, Commands](FToolMenuSection& InSection)
+			{
+				const UWaveformEditorToolMenuContext* Context = InSection.FindContext<UWaveformEditorToolMenuContext>();
+
+				if (Context && Context->WaveformEditor.IsValid())
+				{
+					FToolMenuEntry DeleteLoopRegion = FToolMenuEntry::InitToolBarButton(
+						Commands.DeleteMarker,
+						FText(),
+						LOCTEXT("WaveformEditorDeleteMarkerButtonTooltip", "Delete a marker cue or loop region for the soundwave"),
+						FSlateIcon(FAppStyle::GetAppStyleSetName(), "Symbols.X")
+					);
+					InSection.AddEntry(DeleteLoopRegion);
+				}
+			}));
 	}
 
 	return true;
@@ -951,6 +1055,18 @@ bool FWaveformEditor::BindCommands()
 		Commands.SelectPreviousLoop,
 		FExecuteAction::CreateLambda([this] { CycleMarkerLoopRegion(ELoopModificationControls::SelectPreviousLoop); }));
 
+	ToolkitCommands->MapAction(
+		Commands.CreateMarker,
+		FExecuteAction::CreateLambda([this] { FWaveformEditor::CreateMarker(false); }));
+
+	ToolkitCommands->MapAction(
+		Commands.CreateLoopRegion,
+		FExecuteAction::CreateLambda([this] { FWaveformEditor::CreateMarker(true); }));
+
+	ToolkitCommands->MapAction(
+		Commands.DeleteMarker,
+		FExecuteAction::CreateSP(this, &FWaveformEditor::DeleteMarker));
+
 	return true;
 }
 
diff --git a/Engine/Plugins/Editor/WaveformEditor/Source/WaveformEditor/Private/WaveformEditorCommands.cpp b/Engine/Plugins/Editor/WaveformEditor/Source/WaveformEditor/Private/WaveformEditorCommands.cpp
--- a/Engine/Plugins/Editor/WaveformEditor/Source/WaveformEditor/Private/WaveformEditorCommands.cpp
+++ b/Engine/Plugins/Editor/WaveformEditor/Source/WaveformEditor/Private/WaveformEditorCommands.cpp
@@ -53,6 +53,10 @@ void FWaveformEditorCommands::RegisterCommands()
 	UI_COMMAND(BoundsIncrementDecrease, "Decrease Bound Increment", "Increases the increment to shift the loop bounds", EUserInterfaceActionType::Button, FInputChord(EKeys::Hyphen));
 	UI_COMMAND(SelectNextLoop, "Select Next Loop", "Selects the next loop marker", EUserInterfaceActionType::Button, FInputChord(EKeys::F));
 	UI_COMMAND(SelectPreviousLoop, "Select Previous Loop", "Selects the previous loop marker", EUserInterfaceActionType::Button, FInputChord(EKeys::B));
+
+	UI_COMMAND(CreateMarker, "Create Marker", "Toggle adding Marker Cue Transformation onto the sound wave", EUserInterfaceActionType::Button, FInputChord());
+	UI_COMMAND(CreateLoopRegion, "Create Loop Region", "Toggle adding Marker Loop Region Transformation onto the sound wave", EUserInterfaceActionType::Button, FInputChord());
+	UI_COMMAND(DeleteMarker, "Delete Marker Or Loop Region", "Delete a selected Marker Cue or Loop Region from the sound wave", EUserInterfaceActionType::Button, FInputChord());
 }
 	
 const FWaveformEditorCommands& FWaveformEditorCommands::Get()
diff --git a/Engine/Plugins/Editor/WaveformEditor/Source/WaveformEditor/Public/WaveformEditor.h b/Engine/Plugins/Editor/WaveformEditor/Source/WaveformEditor/Public/WaveformEditor.h
--- a/Engine/Plugins/Editor/WaveformEditor/Source/WaveformEditor/Public/WaveformEditor.h
+++ b/Engine/Plugins/Editor/WaveformEditor/Source/WaveformEditor/Public/WaveformEditor.h
@@ -10,7 +10,7 @@
 #include "WaveformEditorZoomController.h"
 #include "WaveformTransformationTrimFade.h"
 #include "TransformedWaveformView.h"
-#include <WaveformTransformationMarkers.h>
+#include "WaveformTransformationMarkers.h"
 
 class IToolkitHost;
 class SDockTab;
@@ -106,13 +106,15 @@ private:
 	void NotifyPostTransformationChange(const EPropertyChangeType::Type& PropertyChangeType = EPropertyChangeType::ArrayAdd);
 	TObjectPtr<UWaveformTransformationTrimFade> GetOrAddTrimFadeTransformation();
 
+	TObjectPtr<UWaveformTransformationMarkers> GetOrAddMarkerTransformation();
 	void ToggleFadeIn();
 	bool CanFadeIn();
 
 	void ToggleFadeOut();
 	bool CanFadeOut();
+	void CreateMarker(bool bIsLoopRegion);
+	void DeleteMarker();
 
-	TObjectPtr<UWaveformTransformationMarkers> GetOrAddMarkerTransformation();
 	void ModifyMarkerLoopRegion(ELoopModificationControls Modification);
 	void CycleMarkerLoopRegion(ELoopModificationControls Modification);
 
diff --git a/Engine/Plugins/Editor/WaveformEditor/Source/WaveformEditor/Public/WaveformEditorCommands.h b/Engine/Plugins/Editor/WaveformEditor/Source/WaveformEditor/Public/WaveformEditorCommands.h
--- a/Engine/Plugins/Editor/WaveformEditor/Source/WaveformEditor/Public/WaveformEditorCommands.h
+++ b/Engine/Plugins/Editor/WaveformEditor/Source/WaveformEditor/Public/WaveformEditorCommands.h
@@ -52,4 +52,8 @@ public:
 	TSharedPtr<FUICommandInfo> BoundsIncrementDecrease;
 	TSharedPtr<FUICommandInfo> SelectNextLoop;
 	TSharedPtr<FUICommandInfo> SelectPreviousLoop;
+
+	TSharedPtr<FUICommandInfo> CreateMarker;
+	TSharedPtr<FUICommandInfo> CreateLoopRegion;
+	TSharedPtr<FUICommandInfo> DeleteMarker;
 };
diff --git a/Engine/Plugins/Editor/WaveformEditor/Source/WaveformTransformations/Public/WaveformTransformationMarkers.h b/Engine/Plugins/Editor/WaveformEditor/Source/WaveformTransformations/Public/WaveformTransformationMarkers.h
--- a/Engine/Plugins/Editor/WaveformEditor/Source/WaveformTransformations/Public/WaveformTransformationMarkers.h
+++ b/Engine/Plugins/Editor/WaveformEditor/Source/WaveformTransformations/Public/WaveformTransformationMarkers.h
@@ -38,6 +38,7 @@ public:
 
 	UPROPERTY(EditAnywhere, Category = "Markers")
 	TArray<FSoundWaveCuePoint> WaveCues;
+	int32 SelectedCue = INDEX_NONE;
 
 	DECLARE_DELEGATE_OneParam(ModifyMarkerLoopRegionDelegate, ELoopModificationControls);
 	ModifyMarkerLoopRegionDelegate ModifyMarkerLoop;
diff --git a/Engine/Plugins/Editor/WaveformEditor/Source/WaveformTransformationsWidgets/Private/WaveformTransformationMarkersRenderer.cpp b/Engine/Plugins/Editor/WaveformEditor/Source/WaveformTransformationsWidgets/Private/WaveformTransformationMarkersRenderer.cpp
--- a/Engine/Plugins/Editor/WaveformEditor/Source/WaveformTransformationsWidgets/Private/WaveformTransformationMarkersRenderer.cpp
+++ b/Engine/Plugins/Editor/WaveformEditor/Source/WaveformTransformationsWidgets/Private/WaveformTransformationMarkersRenderer.cpp
@@ -96,7 +96,7 @@ int32 FWaveformTransformationMarkerRenderer::DrawMarkerHandles(const FGeometry&
 		float colorMultiplier = 0.5;
 
 		// Highlight selection
-		if (SelectedCue != INDEX_NONE && SelectedCue == LoopRegion.CuePointID)
+		if (MarkersArray->SelectedCue != INDEX_NONE && MarkersArray->SelectedCue == LoopRegion.CuePointID)
 		{
 			colorMultiplier = 0.9;
 		}
@@ -195,7 +195,7 @@ int32 FWaveformTransformationMarkerRenderer::DrawMarkerHandles(const FGeometry&
 			HandleGeometry,
 			&RoundedBoxBrush,
 			ESlateDrawEffect::None,
-			MarkerColor
+			MarkersArray->SelectedCue == Marker.CuePointID ? SelectedMarkerColor : MarkerColor
 		);
 
 		FVector2D TextOffset(MarkerPosition + MarkerWidth, 0);
@@ -302,18 +302,18 @@ void FWaveformTransformationMarkerRenderer::CycleMarkerLoopRegion(ELoopModificat
 	const int32 MarkerArrayLength = MarkersArray->WaveCues.Num();
 
 	// If there is no active selection, pick first or last elmeent
-	if (SelectedCue == INDEX_NONE && MarkerArrayLength > 0)
+	if (MarkersArray->SelectedCue == INDEX_NONE && MarkerArrayLength > 0)
 	{
 		switch (Modification)
 		{
 		case ELoopModificationControls::SelectNextLoop:
 		{
-			SelectedCue = MarkersArray->WaveCues[0].CuePointID;
+			MarkersArray->SelectedCue = MarkersArray->WaveCues[0].CuePointID;
 		}
 			break;
 		case ELoopModificationControls::SelectPreviousLoop:
 		{
-			SelectedCue = MarkersArray->WaveCues[MarkerArrayLength - 1].CuePointID;
+			MarkersArray->SelectedCue = MarkersArray->WaveCues[MarkerArrayLength - 1].CuePointID;
 		}
 			break;
 		default:
@@ -324,17 +324,17 @@ void FWaveformTransformationMarkerRenderer::CycleMarkerLoopRegion(ELoopModificat
 	// Since cues can be removed and added the ids do not match list indexes, so search for matches
 	for (int i = 0; i < MarkerArrayLength; i++)
 	{
-		if (MarkersArray->WaveCues[i].CuePointID == SelectedCue)
+		if (MarkersArray->WaveCues[i].CuePointID == MarkersArray->SelectedCue)
 		{
 			// Check if we move backwards of forwards
 			if (Modification == ELoopModificationControls::SelectNextLoop && i < MarkerArrayLength - 1)
 			{
-				SelectedCue = MarkersArray->WaveCues[i + 1].CuePointID;
+				MarkersArray->SelectedCue = MarkersArray->WaveCues[i + 1].CuePointID;
 				break;
 			}
 			else if (Modification == ELoopModificationControls::SelectPreviousLoop && i > 0)
 			{
-				SelectedCue = MarkersArray->WaveCues[i - 1].CuePointID;
+				MarkersArray->SelectedCue = MarkersArray->WaveCues[i - 1].CuePointID;
 				break;
 			}
 		}
@@ -346,14 +346,14 @@ FSoundWaveCuePoint* FWaveformTransformationMarkerRenderer::GetSelectedMarker() c
 {
 	check(MarkersArray);
 
-	if (SelectedCue == INDEX_NONE)
+	if (MarkersArray->SelectedCue == INDEX_NONE)
 	{
 		return nullptr;
 	}
 
 	for (int i = 0; i < MarkersArray->WaveCues.Num(); i++)
 	{
-		if (MarkersArray->WaveCues[i].CuePointID == SelectedCue)
+		if (MarkersArray->WaveCues[i].CuePointID == MarkersArray->SelectedCue)
 		{
 			return &MarkersArray->WaveCues[i];
 		}
@@ -376,7 +376,7 @@ void FWaveformTransformationMarkerRenderer::SelectLoopRegionByKeyboard(const FKe
 
 		if (PressedKey == IndexKeyCodes[i])
 		{
-			SelectedCue = i;
+			MarkersArray->SelectedCue = i;
 		}
 	}
 }
@@ -676,6 +676,7 @@ bool FWaveformTransformationMarkerRenderer::SetMarkerInInteractionRange(const FV
 		if (IsPositionInInteractionRange(CuePoint.FramePosition, InLocalCursorPosition, WidgetGeometry))
 		{
 			MarkerInInteractionRange.SetMarkerInInteractionRange(&CuePoint);
+			MarkersArray->SelectedCue = CuePoint.CuePointID;
 			
 			return true;
 		}
@@ -705,7 +706,7 @@ bool FWaveformTransformationMarkerRenderer::SetMarkerInInteractionRange(const FV
 			int32 LoopRegionPosition = (InLocalCursorPosition.X / PixelsPerFrame) - StartFrameOffset - CuePoint.FramePosition;
 
 			MarkerInInteractionRange.SetMarkerInInteractionRange(&CuePoint, false, LoopRegionPosition);
-			SelectedCue = CuePoint.CuePointID;
+			MarkersArray->SelectedCue = CuePoint.CuePointID;
 
 			return true;
 		}
@@ -746,7 +747,7 @@ void FWaveformTransformationMarkerRenderer::SetActiveLoopRegion(int32 CuePointID
 
 	MarkersHandle->NotifyPreChange();
 	MarkersArray->Modify();
-	SelectedCue = CuePointID;
+	MarkersArray->SelectedCue = CuePointID;
 	EndTransaction();
 	MarkersHandle->NotifyPostChange(EPropertyChangeType::ValueSet);
 }
diff --git a/Engine/Plugins/Editor/WaveformEditor/Source/WaveformTransformationsWidgets/Private/WaveformTransformationMarkersRenderer.h b/Engine/Plugins/Editor/WaveformEditor/Source/WaveformTransformationsWidgets/Private/WaveformTransformationMarkersRenderer.h
--- a/Engine/Plugins/Editor/WaveformEditor/Source/WaveformTransformationsWidgets/Private/WaveformTransformationMarkersRenderer.h
+++ b/Engine/Plugins/Editor/WaveformEditor/Source/WaveformTransformationsWidgets/Private/WaveformTransformationMarkersRenderer.h
@@ -89,10 +89,10 @@ private:
 
 	FVector2D MousePosition;
 	FSoundWaveCuePoint* MarkerToModify = nullptr;
-	int32 SelectedCue = INDEX_NONE;
 	int32 HighlightedCue = INDEX_NONE;
 
 	FLinearColor MarkerColor = WaveformTransformationWidgetSharedDefaults::DefaultMarkerColor;
+	FLinearColor SelectedMarkerColor = WaveformTransformationWidgetSharedDefaults::DefaultSelectedMarkerColor;
 	TArray<FLinearColor> LoopColors = WaveformTransformationWidgetSharedDefaults::DefaultLoopColors;
 	FLinearColor LabelTextColor = WaveformTransformationWidgetSharedDefaults::DefaultLabelTextColor;
 	float LabelFontSize = WaveformTransformationWidgetSharedDefaults::DefaultLabelFontSize;
diff --git a/Engine/Plugins/Editor/WaveformEditor/Source/WaveformTransformationsWidgets/Private/WaveformTransformationsWidgetsSettings.h b/Engine/Plugins/Editor/WaveformEditor/Source/WaveformTransformationsWidgets/Private/WaveformTransformationsWidgetsSettings.h
--- a/Engine/Plugins/Editor/WaveformEditor/Source/WaveformTransformationsWidgets/Private/WaveformTransformationsWidgetsSettings.h
+++ b/Engine/Plugins/Editor/WaveformEditor/Source/WaveformTransformationsWidgets/Private/WaveformTransformationsWidgetsSettings.h
@@ -10,7 +10,8 @@ DECLARE_MULTICAST_DELEGATE_TwoParams(FOnWaveformTransformationsWidgetsSettingsCh
 
 namespace WaveformTransformationWidgetSharedDefaults
 {
-	const FLinearColor DefaultMarkerColor			= FLinearColor(1.f, 0.9f, 0, 0.75f);
+	const FLinearColor DefaultMarkerColor			= FLinearColor(1.f, 0.9f, 0, 0.5f);
+	const FLinearColor DefaultSelectedMarkerColor = FLinearColor(1.f, 0.9f, 0, 1.f);
 	const TArray<FLinearColor> DefaultLoopColors	= { FLinearColor(0.8f, 0.6f, 1.f, 0.75f), FLinearColor(0.2f, 0.7f, 1.f, 0.75f) };
 	const FLinearColor DefaultLabelTextColor		= FLinearColor(0.9f, 0.9f, 0.9, 1.f);
 	const float DefaultLabelFontSize				= 10.f;
