[Waveform Editor] - Editor crashes when opening a level while the wave editor is open and has a transform on it

diff --git a/Engine/Plugins/Editor/WaveformEditor/Source/WaveformEditor/Private/WaveformEditor.cpp b/Engine/Plugins/Editor/WaveformEditor/Source/WaveformEditor/Private/WaveformEditor.cpp
--- a/Engine/Plugins/Editor/WaveformEditor/Source/WaveformEditor/Private/WaveformEditor.cpp
+++ b/Engine/Plugins/Editor/WaveformEditor/Source/WaveformEditor/Private/WaveformEditor.cpp
@@ -101,6 +101,18 @@ bool FWaveformEditor::Init(const EToolkitMode::Type Mode, const TSharedPtr<ITool
 			TransportController->CacheStartTime(PlaybackTimeBeforeTransformInteraction);
 			check(TransportCoordinator != nullptr);
 			TransportCoordinator->SetProgressRatio(0.0f);
+
+			OnMapOpenedHandle = FEditorDelegates::OnMapOpened.AddLambda([this](const FString&, bool)
+				{
+					if (TransformationsDetails.IsValid())
+					{
+						TransformationsDetails->ForceRefresh();
+					}
+					if (PropertiesDetails.IsValid())
+					{
+						PropertiesDetails->ForceRefresh();
+					}
+				});
 	}
 
 	return bIsInitialized;
@@ -108,6 +120,12 @@ bool FWaveformEditor::Init(const EToolkitMode::Type Mode, const TSharedPtr<ITool
 
 FWaveformEditor::~FWaveformEditor()
 {
+	if (OnMapOpenedHandle.IsValid())
+	{
+		FEditorDelegates::OnMapOpened.Remove(OnMapOpenedHandle);
+		OnMapOpenedHandle.Reset();
+	}
+
 	if (FReimportManager::Instance())
 	{
 		FReimportManager::Instance()->OnPostReimport().RemoveAll(this);
diff --git a/Engine/Plugins/Editor/WaveformEditor/Source/WaveformEditor/Public/WaveformEditor.h b/Engine/Plugins/Editor/WaveformEditor/Source/WaveformEditor/Public/WaveformEditor.h
--- a/Engine/Plugins/Editor/WaveformEditor/Source/WaveformEditor/Public/WaveformEditor.h
+++ b/Engine/Plugins/Editor/WaveformEditor/Source/WaveformEditor/Public/WaveformEditor.h
@@ -176,5 +176,7 @@ private:
 
 	FText GetReimportButtonToolTip() const;
 	FText GetExportButtonToolTip() const;
+
+	FDelegateHandle OnMapOpenedHandle;
 };
 
