Waveform Editor: Ability to tune and preview loop points
Using IPropertyTypeCustomization to edit the details view of the markers transformation and remove unneeded information in the waveform editor

diff --git a/Engine/Plugins/Editor/WaveformEditor/Config/BaseWaveformEditor.ini b/Engine/Plugins/Editor/WaveformEditor/Config/BaseWaveformEditor.ini
--- a/Engine/Plugins/Editor/WaveformEditor/Config/BaseWaveformEditor.ini
+++ b/Engine/Plugins/Editor/WaveformEditor/Config/BaseWaveformEditor.ini
@@ -9,3 +9,6 @@ RulerTicksColor=(R=1.000000,G=1.000000,B=1.000000,A=0.900000)
 RulerTextColor=(R=1.000000,G=1.000000,B=1.000000,A=0.900000)
 RulerFontSize=10.000000
 
+
+[CoreRedirects]
++PropertyRedirects=(OldName="/Script/WaveformTransformations.WaveCueArray.WaveCues",NewName="/Script/WaveformTransformations.WaveCueArray.CuesAndLoops")
diff --git a/Engine/Plugins/Editor/WaveformEditor/Source/WaveformEditor/Private/WaveformEditor.cpp b/Engine/Plugins/Editor/WaveformEditor/Source/WaveformEditor/Private/WaveformEditor.cpp
--- a/Engine/Plugins/Editor/WaveformEditor/Source/WaveformEditor/Private/WaveformEditor.cpp
+++ b/Engine/Plugins/Editor/WaveformEditor/Source/WaveformEditor/Private/WaveformEditor.cpp
@@ -130,6 +130,15 @@ FWaveformEditor::~FWaveformEditor()
 	{
 		FReimportManager::Instance()->OnPostReimport().RemoveAll(this);
 	}
+
+	if (TransformationChainConfig.bIsPreviewingLoopRegion)
+	{
+		// Note: if the user doesn't save the soundwave asset, bLooping will not be reset.
+		check(SoundWave);
+		SoundWave->Modify();
+		SoundWave->bLooping = TransformationChainConfig.bCachedSoundWaveLoopState;
+		SoundWave->PostEditChange();
+	}
 }
 
 void FWaveformEditor::AddDefaultTransformations()
@@ -350,7 +359,7 @@ void FWaveformEditor::CreateMarker(bool bIsLoopRegion)
 
 	SoundWave->Modify();
 
-	for (FSoundWaveCuePoint CuePoint : MarkerTransformation->Markers->WaveCues)
+	for (FSoundWaveCuePoint CuePoint : MarkerTransformation->Markers->CuesAndLoops)
 	{
 		if (CuePoint.CuePointID > HighestCueID)
 		{
@@ -368,7 +377,7 @@ void FWaveformEditor::CreateMarker(bool bIsLoopRegion)
 		NewCuePoint.FrameLength = SoundWave->TotalSamples * 0.1; // Default loop region to 10% of samples for easier tuning
 	}
 
-	MarkerTransformation->Markers->WaveCues.Add(NewCuePoint);
+	MarkerTransformation->Markers->CuesAndLoops.Add(NewCuePoint);
 	NotifyPostTransformationChange(EPropertyChangeType::ValueSet);
 }
 
@@ -384,11 +393,11 @@ void FWaveformEditor::DeleteMarker()
 		return; // No cue to selected
 	}
 
-	for (int i = 0; i < MarkerTransformation->Markers->WaveCues.Num(); i++)
+	for (int i = 0; i < MarkerTransformation->Markers->CuesAndLoops.Num(); i++)
 	{
-		if (CueToDelete == MarkerTransformation->Markers->WaveCues[i].CuePointID)
+		if (CueToDelete == MarkerTransformation->Markers->CuesAndLoops[i].CuePointID)
 		{
-			MarkerTransformation->Markers->WaveCues.RemoveAt(i);
+			MarkerTransformation->Markers->CuesAndLoops.RemoveAt(i);
 			MarkerTransformation->Markers->SelectedCue = INDEX_NONE;
 			break;
 		}
diff --git a/Engine/Plugins/Editor/WaveformEditor/Source/WaveformEditorWidgets/Private/WaveformEditorSequenceDataProvider.cpp b/Engine/Plugins/Editor/WaveformEditor/Source/WaveformEditorWidgets/Private/WaveformEditorSequenceDataProvider.cpp
--- a/Engine/Plugins/Editor/WaveformEditor/Source/WaveformEditorWidgets/Private/WaveformEditorSequenceDataProvider.cpp
+++ b/Engine/Plugins/Editor/WaveformEditor/Source/WaveformEditorWidgets/Private/WaveformEditorSequenceDataProvider.cpp
@@ -120,7 +120,6 @@ void FWaveformEditorSequenceDataProvider::GenerateSequenceDataInternal()
 		TransformationInfo.SampleRate = ImportedSampleRate;
 
 		TArray<Audio::FTransformationPtr> Transformations = CreateTransformations();
-
 		const bool bChainChangesFileLength = CanChainChangeFileLength(Transformations);
 
 		for (int32 i = 0; i < Transformations.Num(); ++i)
@@ -139,17 +138,23 @@ void FWaveformEditorSequenceDataProvider::GenerateSequenceDataInternal()
 
 			}
 
-			FirstEditedSample += TransformationInfo.StartFrameOffset;
-			LastEditedSample = TransformationInfo.NumEditedSamples <= 0 ? LastEditedSample : FirstEditedSample + TransformationInfo.NumEditedSamples;
-
 			if (bChainChangesFileLength)
 			{
+				FirstEditedSample += TransformationInfo.StartFrameOffset;
+
+				//Prevent a lower priority transformation from setting the LastEditedSample past the end of higher priority transformations
+				if (Transformations[i]->FileChangeLengthPriority() == Audio::ETransformationPriority::High
+					|| (Transformations[i]->FileChangeLengthPriority() != Audio::ETransformationPriority::None
+						&& FirstEditedSample + TransformationInfo.NumEditedSamples < LastEditedSample))
+				{
+					LastEditedSample = TransformationInfo.NumEditedSamples <= 0 ? LastEditedSample : FirstEditedSample + TransformationInfo.NumEditedSamples;
+				}
+			
 				check(LastEditedSample > FirstEditedSample);
 				const int32 EditedBufferSize = LastEditedSample - FirstEditedSample;
 				FMemory::Memcpy(&OutputBuffer.GetData()[FirstEditedSample], TransformationsBuffer.GetData(), EditedBufferSize * sizeof(float));
 			}
 
-			TransformationInfo.StartFrameOffset = 0;
 			TransformationInfo.NumEditedSamples = 0;
 		}
 
@@ -214,7 +219,10 @@ TArray<Audio::FTransformationPtr> FWaveformEditorSequenceDataProvider::CreateTra
 
 	for (const TObjectPtr<UWaveformTransformationBase>& TransformationBase : TransformationsToRender)
 	{
-		TransformationPtrs.Add(TransformationBase->CreateTransformation());
+		if(TransformationBase)
+		{
+			TransformationPtrs.Add(TransformationBase->CreateTransformation());
+		}
 	}
 
 	return TransformationPtrs;
@@ -226,7 +234,7 @@ const bool FWaveformEditorSequenceDataProvider::CanChainChangeFileLength(const T
 
 	for (const Audio::FTransformationPtr& Transformation : TransformationChain)
 	{
-		bCanChainChangeFileLength |= Transformation->CanChangeFileLength();
+		bCanChainChangeFileLength |= Transformation->FileChangeLengthPriority() != Audio::ETransformationPriority::None;
 	}
 
 	return bCanChainChangeFileLength;
diff --git a/Engine/Plugins/Editor/WaveformEditor/Source/WaveformTransformations/Private/WaveformTransformationMarkers.cpp b/Engine/Plugins/Editor/WaveformEditor/Source/WaveformTransformations/Private/WaveformTransformationMarkers.cpp
--- a/Engine/Plugins/Editor/WaveformEditor/Source/WaveformTransformations/Private/WaveformTransformationMarkers.cpp
+++ b/Engine/Plugins/Editor/WaveformEditor/Source/WaveformTransformations/Private/WaveformTransformationMarkers.cpp
@@ -3,6 +3,76 @@
 #include "WaveformTransformationMarkers.h"
 #include "Sound/SoundWave.h"
 
+DEFINE_LOG_CATEGORY(LogWaveformTransformationMarkers);
+
+FWaveTransformationMarkers::FWaveTransformationMarkers(double InStartLoopTime, double InEndLoopTime)
+: StartLoopTime(InStartLoopTime)
+, EndLoopTime(InEndLoopTime){}
+
+void FWaveTransformationMarkers::ProcessAudio(Audio::FWaveformTransformationWaveInfo& InOutWaveInfo) const
+{
+	check(InOutWaveInfo.SampleRate > 0.f && InOutWaveInfo.Audio != nullptr);
+
+	Audio::FAlignedFloatBuffer& InputAudio = *InOutWaveInfo.Audio;
+
+	if (InputAudio.Num() == 0)
+	{
+		return;
+	}
+
+	check(StartLoopTime >= 0);
+	check(InOutWaveInfo.NumChannels > 0);
+
+	const int32 StartSampleOffset = InOutWaveInfo.StartFrameOffset;
+
+	int32 StartSample = FMath::RoundToInt32(StartLoopTime * InOutWaveInfo.SampleRate) * InOutWaveInfo.NumChannels;
+	StartSample = FMath::Max(0, StartSample - StartSampleOffset);
+
+	if (StartSample > InputAudio.Num())
+	{
+		return;
+	}
+	
+	int32 EndSample = InputAudio.Num() - 1;
+
+	if (EndLoopTime > 0.f)
+	{
+		const int32 EndFrame = FMath::RoundToInt32(EndLoopTime * InOutWaveInfo.SampleRate) * InOutWaveInfo.NumChannels;
+		EndSample = FMath::Max(StartSample, EndFrame - 1 - StartSampleOffset);
+
+		// EndLoopTime can be beyond the length of the file if there is a Trim
+		if (EndSample > InputAudio.Num() - 1)
+		{
+			EndSample = InputAudio.Num() - 1;
+			UE_LOG(LogWaveformTransformationMarkers, Warning, TEXT("Cutting a loop point with a trim!"));
+		}
+
+		EndSample = FMath::Min(EndSample, InputAudio.Num() - 1);
+	}
+
+	const int32 FinalSize = EndSample - StartSample + 1;
+
+	InOutWaveInfo.StartFrameOffset = StartSample - (StartSample % InOutWaveInfo.NumChannels);
+	InOutWaveInfo.NumEditedSamples = FinalSize;
+
+	// Audio needs no trimming
+	if (FinalSize == InputAudio.Num())
+	{
+		return;
+	}
+
+	// Apply trim to the audio to audition the desired loop region
+	TArray<float> TempBuffer;
+	TempBuffer.AddUninitialized(FinalSize);
+
+	FMemory::Memcpy(TempBuffer.GetData(), &InputAudio[StartSample], FinalSize * sizeof(float));
+
+	InputAudio.Empty();
+	InputAudio.AddUninitialized(FinalSize);
+
+	FMemory::Memcpy(InputAudio.GetData(), TempBuffer.GetData(), FinalSize * sizeof(float));
+}
+
 UWaveformTransformationMarkers::UWaveformTransformationMarkers(const FObjectInitializer& ObjectInitializer) : Super(ObjectInitializer)
 {
 	if(Markers == nullptr)
@@ -13,13 +83,29 @@ UWaveformTransformationMarkers::UWaveformTransformationMarkers(const FObjectInit
 
 Audio::FTransformationPtr UWaveformTransformationMarkers::CreateTransformation() const
 {
-	return MakeUnique<FWaveTransformationMarkers>();
+	return MakeUnique<FWaveTransformationMarkers>(StartLoopTime, EndLoopTime);
 }
 
 void UWaveformTransformationMarkers::UpdateConfiguration(FWaveTransformUObjectConfiguration& InOutConfiguration)
 {
 	check(Markers != nullptr);
 	Markers->InitMarkersIfNotSet(InOutConfiguration.WaveCues);
+
+	if (!bCachedIsPreviewingLoopRegion)
+	{
+		bCachedSoundWaveLoopState = InOutConfiguration.bCachedSoundWaveLoopState;
+		EndLoopTime = -1;
+	}
+	else
+	{
+		InOutConfiguration.StartTime = StartLoopTime;
+	}
+
+	// Update after setting bCachedSoundWaveLoopState so bCachedSoundWaveLoopState isn't overwritten
+	bCachedIsPreviewingLoopRegion = bIsPreviewingLoopRegion;
+
+	InOutConfiguration.bIsPreviewingLoopRegion = bIsPreviewingLoopRegion;
+	InOutConfiguration.bCachedSoundWaveLoopState = bCachedSoundWaveLoopState;
 }
 
 void UWaveformTransformationMarkers::OverwriteTransformation()
@@ -65,13 +151,13 @@ void UWaveformTransformationMarkers::OverwriteSoundWaveData(USoundWave& InOutSou
 	// If this is the first Marker transformation, overwrite the soundwave original SoundWaveCuePoints
 	if (ThisTransformationNumber == 1)
 	{
-		InOutSoundWave.SetSoundWaveCuePoints(Markers->WaveCues);
+		InOutSoundWave.SetSoundWaveCuePoints(Markers->CuesAndLoops);
 	}
 	// Handles the case if there are multiple marker transformations to preserve all marker "tracks"
 	else
 	{
 		TArray<FSoundWaveCuePoint> SoundWaveCuePoints(InOutSoundWave.GetSoundWaveCuePoints());
-		SoundWaveCuePoints.Append(Markers->WaveCues);
+		SoundWaveCuePoints.Append(Markers->CuesAndLoops);
 		
 		InOutSoundWave.SetSoundWaveCuePoints(SoundWaveCuePoints);
 	}
@@ -80,7 +166,7 @@ void UWaveformTransformationMarkers::OverwriteSoundWaveData(USoundWave& InOutSou
 void UWaveformTransformationMarkers::GetTransformationInfo(FWaveformTransformationInfo& InOutTransformationInfo) const
 {
 	check(Markers != nullptr);
-	InOutTransformationInfo.AllCuePoints.Append(Markers->WaveCues);
+	InOutTransformationInfo.AllCuePoints.Append(Markers->CuesAndLoops);
 }
 
 void UWaveCueArray::PostEditChangeProperty(FPropertyChangedEvent& PropertyChangedEvent)
@@ -88,20 +174,20 @@ void UWaveCueArray::PostEditChangeProperty(FPropertyChangedEvent& PropertyChange
 	Super::PostEditChangeProperty(PropertyChangedEvent);
 
 	const FName PropertyName = PropertyChangedEvent.GetPropertyName();
-	if (PropertyName == GET_MEMBER_NAME_CHECKED(UWaveCueArray, WaveCues) && WaveCues.Num() > 0 &&
+	if (PropertyName == GET_MEMBER_NAME_CHECKED(UWaveCueArray, CuesAndLoops) && CuesAndLoops.Num() > 0 &&
 		(PropertyChangedEvent.ChangeType & (EPropertyChangeType::ArrayAdd | EPropertyChangeType::ResetToDefault | 
 			EPropertyChangeType::Unspecified | EPropertyChangeType::Redirected)))
 	{
 		int32 LargestCuePointID = 1;
-		for (int32 Index = WaveCues.Num()-1; Index >= 0; Index--)
+		for (int32 Index = CuesAndLoops.Num()-1; Index >= 0; Index--)
 		{
-			if (WaveCues[Index].CuePointID >= LargestCuePointID)
+			if (CuesAndLoops[Index].CuePointID >= LargestCuePointID)
 			{
-				LargestCuePointID = WaveCues[Index].CuePointID + 1;
+				LargestCuePointID = CuesAndLoops[Index].CuePointID + 1;
 			}
 		}
 		//Elements can be inserted at any index in the array
-		for (FSoundWaveCuePoint& CuePoint : WaveCues)
+		for (FSoundWaveCuePoint& CuePoint : CuesAndLoops)
 		{
 			if (CuePoint.CuePointID == INDEX_NONE)
 			{
@@ -116,14 +202,14 @@ void UWaveCueArray::InitMarkersIfNotSet(const TArray<FSoundWaveCuePoint>& InMark
 	//Prevent USoundWave from overwriting the Transformation unintentionally
 	if (!bIsInitialized)
 	{
-		WaveCues = InMarkers;
+		CuesAndLoops = InMarkers;
 		bIsInitialized = true;
 	}
 }
 
 void UWaveCueArray::Reset()
 {
-	WaveCues.Empty();
+	CuesAndLoops.Empty();
 	bIsInitialized = false;
 }
 #endif
diff --git a/Engine/Plugins/Editor/WaveformEditor/Source/WaveformTransformations/Private/WaveformTransformationMarkersObjectCustomization.cpp b/Engine/Plugins/Editor/WaveformEditor/Source/WaveformTransformations/Private/WaveformTransformationMarkersObjectCustomization.cpp
new file mode 100644
--- /dev/null
+++ b/Engine/Plugins/Editor/WaveformEditor/Source/WaveformTransformations/Private/WaveformTransformationMarkersObjectCustomization.cpp
@@ -0,0 +1,39 @@
+// Copyright Epic Games, Inc. All Rights Reserved.
+
+#include "WaveformTransformationMarkersObjectCustomization.h"
+
+#include "DetailWidgetRow.h"
+#include "IDetailChildrenBuilder.h"
+#include "IPropertyUtilities.h"
+#include "PropertyHandle.h"
+#include "WaveformTransformationMarkers.h"
+
+TSharedRef<IPropertyTypeCustomization> FWaveformTransformationMarkersObjectCustomization::MakeInstance()
+{
+	return MakeShareable(new FWaveformTransformationMarkersObjectCustomization);
+}
+
+void FWaveformTransformationMarkersObjectCustomization::CustomizeHeader(TSharedRef<IPropertyHandle> PropertyHandle,
+	FDetailWidgetRow& HeaderRow, IPropertyTypeCustomizationUtils& CustomizationUtils)
+{
+	HeaderRow.NameContent()[PropertyHandle->CreatePropertyNameWidget()]
+	.ValueContent()[PropertyHandle->CreatePropertyValueWidget()];
+}
+
+void FWaveformTransformationMarkersObjectCustomization::CustomizeChildren(TSharedRef<IPropertyHandle> PropertyHandle,
+	IDetailChildrenBuilder& ChildBuilder, IPropertyTypeCustomizationUtils& CustomizationUtils)
+{
+	TSharedPtr<IPropertyHandle> MarkersHandle = PropertyHandle->GetChildHandle(GET_MEMBER_NAME_CHECKED(UWaveformTransformationMarkers, Markers));
+
+	if (!MarkersHandle || !MarkersHandle->IsValidHandle())
+	{
+		return;
+	}
+
+	// Only show the UWaveCueArray property, and nothing else
+	TSharedPtr<IPropertyHandle> WaveCueMarkersHandle = MarkersHandle->GetChildHandle(GET_MEMBER_NAME_CHECKED(UWaveCueArray, CuesAndLoops));
+	if (WaveCueMarkersHandle && WaveCueMarkersHandle->IsValidHandle())
+	{
+		ChildBuilder.AddProperty(WaveCueMarkersHandle.ToSharedRef());
+	}
+}
\ No newline at end of file
diff --git a/Engine/Plugins/Editor/WaveformEditor/Source/WaveformTransformations/Private/WaveformTransformationMarkersObjectCustomization.h b/Engine/Plugins/Editor/WaveformEditor/Source/WaveformTransformations/Private/WaveformTransformationMarkersObjectCustomization.h
new file mode 100644
--- /dev/null
+++ b/Engine/Plugins/Editor/WaveformEditor/Source/WaveformTransformations/Private/WaveformTransformationMarkersObjectCustomization.h
@@ -0,0 +1,18 @@
+// Copyright Epic Games, Inc. All Rights Reserved.
+
+#pragma once
+#include "IPropertyTypeCustomization.h"
+
+/**
+ * Using IPropertyTypeCustomization to edit the details view of the markers transformation and remove unneeded information in the waveform editor
+ */
+class FWaveformTransformationMarkersObjectCustomization : public IPropertyTypeCustomization
+{
+public:
+	static TSharedRef<IPropertyTypeCustomization> MakeInstance();
+
+	virtual void CustomizeHeader(TSharedRef<IPropertyHandle> PropertyHandle, FDetailWidgetRow& HeaderRow,
+								 IPropertyTypeCustomizationUtils& CustomizationUtils) override;
+	virtual void CustomizeChildren(TSharedRef<IPropertyHandle> PropertyHandle, IDetailChildrenBuilder& ChildBuilder,
+								   IPropertyTypeCustomizationUtils& CustomizationUtils) override;
+};
diff --git a/Engine/Plugins/Editor/WaveformEditor/Source/WaveformTransformations/Private/WaveformTransformationsModule.cpp b/Engine/Plugins/Editor/WaveformEditor/Source/WaveformTransformations/Private/WaveformTransformationsModule.cpp
--- a/Engine/Plugins/Editor/WaveformEditor/Source/WaveformTransformations/Private/WaveformTransformationsModule.cpp
+++ b/Engine/Plugins/Editor/WaveformEditor/Source/WaveformTransformations/Private/WaveformTransformationsModule.cpp
@@ -1,6 +1,37 @@
 // Copyright Epic Games, Inc. All Rights Reserved.
 
 #include "WaveformTransformationsModule.h"
+
 #include "Modules/ModuleManager.h"
+#include "PropertyEditorModule.h"
+#include "WaveformTransformationMarkersObjectCustomization.h"
+
+
+void FWaveformTransformationsModule::StartupModule()
+{
+	IModuleInterface::StartupModule();
+	
+	if (FModuleManager::Get().IsModuleLoaded("PropertyEditor"))
+	{
+		FPropertyEditorModule& PropertyEditorModule = FModuleManager::GetModuleChecked<FPropertyEditorModule>("PropertyEditor");
+		
+		PropertyEditorModule.RegisterCustomPropertyTypeLayout(
+			"WaveformTransformationMarkers", 
+			FOnGetPropertyTypeCustomizationInstance::CreateStatic(&FWaveformTransformationMarkersObjectCustomization::MakeInstance)
+		);
+	}
+}
+
+void FWaveformTransformationsModule::ShutdownModule()
+{
+	IModuleInterface::ShutdownModule();
+
+	if (FModuleManager::Get().IsModuleLoaded("PropertyEditor"))
+	{
+		FPropertyEditorModule& PropertyEditorModule = FModuleManager::GetModuleChecked<FPropertyEditorModule>("PropertyEditor");
+
+		PropertyEditorModule.UnregisterCustomClassLayout("WaveformTransformationMarkers");
+	}
+}
 
-IMPLEMENT_MODULE(FWaveformTransformationsModule, WaveformTransformations);
+IMPLEMENT_MODULE(FWaveformTransformationsModule, WaveformTransformations);
\ No newline at end of file
diff --git a/Engine/Plugins/Editor/WaveformEditor/Source/WaveformTransformations/Public/WaveformTransformationEffectChain.h b/Engine/Plugins/Editor/WaveformEditor/Source/WaveformTransformations/Public/WaveformTransformationEffectChain.h
--- a/Engine/Plugins/Editor/WaveformEditor/Source/WaveformTransformations/Public/WaveformTransformationEffectChain.h
+++ b/Engine/Plugins/Editor/WaveformEditor/Source/WaveformTransformations/Public/WaveformTransformationEffectChain.h
@@ -16,7 +16,7 @@ public:
 	virtual void ProcessAudio(Audio::FWaveformTransformationWaveInfo& InOutWaveInfo) const override;
 
 	virtual bool SupportsRealtimePreview() const override { return true; }
-	virtual bool CanChangeFileLength() const override { return true; }
+	virtual constexpr Audio::ETransformationPriority FileChangeLengthPriority() const override { return Audio::ETransformationPriority::High; }
 
 private:
 	TArray<TStrongObjectPtr<class USoundEffectSourcePreset>> Presets;
diff --git a/Engine/Plugins/Editor/WaveformEditor/Source/WaveformTransformations/Public/WaveformTransformationMarkers.h b/Engine/Plugins/Editor/WaveformEditor/Source/WaveformTransformations/Public/WaveformTransformationMarkers.h
--- a/Engine/Plugins/Editor/WaveformEditor/Source/WaveformTransformations/Public/WaveformTransformationMarkers.h
+++ b/Engine/Plugins/Editor/WaveformEditor/Source/WaveformTransformations/Public/WaveformTransformationMarkers.h
@@ -6,6 +6,8 @@
 
 #include "WaveformTransformationMarkers.generated.h"
 
+DECLARE_LOG_CATEGORY_EXTERN(LogWaveformTransformationMarkers, Log, All);
+
 class USoundWave;
 
 enum class ELoopModificationControls : uint8
@@ -29,15 +31,15 @@ class WAVEFORMTRANSFORMATIONS_API UWaveCueArray : public UObject
 public:
 	void PostEditChangeProperty(FPropertyChangedEvent& PropertyChangedEvent);
 	
-	//If uninitialized, init wavecue array
+	//If uninitialized, init Marker array
 	void InitMarkersIfNotSet(const TArray<FSoundWaveCuePoint>& InMarkers);
-	//Uninitialize and empty wavecue array
+	//Uninitialize and empty Marker array
 	void Reset();
 
 	void EnableLoopRegion(FSoundWaveCuePoint* OutSoundWaveCue);
 
 	UPROPERTY(EditAnywhere, Category = "Markers")
-	TArray<FSoundWaveCuePoint> WaveCues;
+	TArray<FSoundWaveCuePoint> CuesAndLoops;
 	int32 SelectedCue = INDEX_NONE;
 
 	DECLARE_DELEGATE_OneParam(ModifyMarkerLoopRegionDelegate, ELoopModificationControls);
@@ -54,7 +56,14 @@ private:
 class WAVEFORMTRANSFORMATIONS_API FWaveTransformationMarkers : public Audio::IWaveTransformation
 {
 public:
-	explicit FWaveTransformationMarkers() = default;
+	explicit FWaveTransformationMarkers(double InStartLoopTime, double InEndLoopTime);
+	virtual void ProcessAudio(Audio::FWaveformTransformationWaveInfo& InOutWaveInfo) const override;
+
+	virtual constexpr Audio::ETransformationPriority FileChangeLengthPriority() const override { return Audio::ETransformationPriority::Low; }
+
+private:
+	double StartLoopTime = 0.0;
+	double EndLoopTime = 0.0;
 };
 
 UCLASS()
@@ -66,11 +75,21 @@ public:
 
 	UPROPERTY(VisibleAnywhere, Category = "Markers")
 	TObjectPtr<UWaveCueArray> Markers;
+
+	UPROPERTY(EditAnywhere, Transient, Category = "Loop Preview", meta = (ClampMin = 0.0))
+	double StartLoopTime = 0.0;
+	UPROPERTY(EditAnywhere, Transient, Category = "Loop Preview")
+	double EndLoopTime = -1.0;
+	UPROPERTY(EditAnywhere, Transient, Category = "Loop Preview")
+	bool bIsPreviewingLoopRegion = false;
+
+	// Default value for EndLoopTime is < 0 to aid with proper initialization in UpdateDurationProperties()
 	
 	virtual Audio::FTransformationPtr CreateTransformation() const override;
 
 	void UpdateConfiguration(FWaveTransformUObjectConfiguration& InOutConfiguration) override;
 	void OverwriteTransformation() override;
+	virtual constexpr Audio::ETransformationPriority GetTransformationPriority() const { return Audio::ETransformationPriority::Low; }
 
 	void ModifyMarkerLoopRegion(ELoopModificationControls Modification) const;
 	void CycleMarkerLoopRegion(ELoopModificationControls Modification) const;
@@ -79,4 +98,10 @@ public:
 	void OverwriteSoundWaveData(USoundWave& InOutSoundWave) override;
 	void GetTransformationInfo(FWaveformTransformationInfo& InOutTransformationInfo) const override;
 #endif //WITH_EDITOR
+
+private:
+	float AvailableWaveformDuration = -1.f;
+
+	bool bCachedIsPreviewingLoopRegion = false;
+	bool bCachedSoundWaveLoopState = false;
 };
diff --git a/Engine/Plugins/Editor/WaveformEditor/Source/WaveformTransformations/Public/WaveformTransformationTrimFade.h b/Engine/Plugins/Editor/WaveformEditor/Source/WaveformTransformations/Public/WaveformTransformationTrimFade.h
--- a/Engine/Plugins/Editor/WaveformEditor/Source/WaveformTransformations/Public/WaveformTransformationTrimFade.h
+++ b/Engine/Plugins/Editor/WaveformEditor/Source/WaveformTransformations/Public/WaveformTransformationTrimFade.h
@@ -22,7 +22,7 @@ public:
 	explicit FWaveTransformationTrimFade(double InStartTime, double InEndTime, float InStartFadeTime, float InStartFadeCurve, float InStartSCurveSharpness, float InEndFadeTime, float InEndFadeCurve, float InEndSCurveSharpness);
 	virtual void ProcessAudio(Audio::FWaveformTransformationWaveInfo& InOutWaveInfo) const override;
 
-	virtual bool CanChangeFileLength() const override { return true; }
+	virtual constexpr Audio::ETransformationPriority FileChangeLengthPriority() const override { return Audio::ETransformationPriority::High; }
 
 private:
 	double StartTime = 0.0;
@@ -72,6 +72,7 @@ public:
 	virtual Audio::FTransformationPtr CreateTransformation() const override;
 
 	void UpdateConfiguration(FWaveTransformUObjectConfiguration& InOutConfiguration) override;
+	virtual constexpr Audio::ETransformationPriority GetTransformationPriority() const { return Audio::ETransformationPriority::High; }
 
 	static const double GetFadeInCurveValue(const float StartFadeCurve, const double FadeFraction, const float SCurveSharpness = 0);
 	static const double GetFadeOutCurveValue(const float EndFadeCurve, const double FadeFraction, const float SCurveSharpness = 0);
diff --git a/Engine/Plugins/Editor/WaveformEditor/Source/WaveformTransformations/Public/WaveformTransformationsModule.h b/Engine/Plugins/Editor/WaveformEditor/Source/WaveformTransformations/Public/WaveformTransformationsModule.h
--- a/Engine/Plugins/Editor/WaveformEditor/Source/WaveformTransformations/Public/WaveformTransformationsModule.h
+++ b/Engine/Plugins/Editor/WaveformEditor/Source/WaveformTransformations/Public/WaveformTransformationsModule.h
@@ -6,4 +6,7 @@
 
 class FWaveformTransformationsModule : public IModuleInterface
 {
+public:
+	virtual void StartupModule() override;
+	virtual void ShutdownModule() override;
 };
diff --git a/Engine/Plugins/Editor/WaveformEditor/Source/WaveformTransformations/WaveformTransformations.build.cs b/Engine/Plugins/Editor/WaveformEditor/Source/WaveformTransformations/WaveformTransformations.build.cs
--- a/Engine/Plugins/Editor/WaveformEditor/Source/WaveformTransformations/WaveformTransformations.build.cs
+++ b/Engine/Plugins/Editor/WaveformEditor/Source/WaveformTransformations/WaveformTransformations.build.cs
@@ -17,6 +17,7 @@ namespace UnrealBuildTool.Rules
 					"Core", 
 					"CoreUObject",
 					"Engine",
+					"PropertyEditor",
 					"SignalProcessing"
 				}
 			);
diff --git a/Engine/Plugins/Editor/WaveformEditor/Source/WaveformTransformationsWidgets/Private/WaveformTransformationMarkersRenderer.cpp b/Engine/Plugins/Editor/WaveformEditor/Source/WaveformTransformationsWidgets/Private/WaveformTransformationMarkersRenderer.cpp
--- a/Engine/Plugins/Editor/WaveformEditor/Source/WaveformTransformationsWidgets/Private/WaveformTransformationMarkersRenderer.cpp
+++ b/Engine/Plugins/Editor/WaveformEditor/Source/WaveformTransformationsWidgets/Private/WaveformTransformationMarkersRenderer.cpp
@@ -218,6 +218,9 @@ void FWaveformTransformationMarkerRenderer::SetPropertyHandles(const TArray<TSha
 	FWaveformTransformationRendererBase::SetPropertyHandles(InPropertyHandles);
 
 	MarkersHandle = GetPropertyHandle(GET_MEMBER_NAME_CHECKED(UWaveformTransformationMarkers, Markers));
+	StartLoopTimeHandle = GetPropertyHandle(GET_MEMBER_NAME_CHECKED(UWaveformTransformationMarkers, StartLoopTime));
+	EndLoopTimeHandle = GetPropertyHandle(GET_MEMBER_NAME_CHECKED(UWaveformTransformationMarkers, EndLoopTime));
+	IsPreviewingLoopRegionHandle = GetPropertyHandle(GET_MEMBER_NAME_CHECKED(UWaveformTransformationMarkers, bIsPreviewingLoopRegion));
 
 	MarkersArray = CastChecked<UWaveCueArray>(GetPropertyValue<UObject*>(MarkersHandle));
 
@@ -299,7 +302,9 @@ void FWaveformTransformationMarkerRenderer::CycleMarkerLoopRegion(ELoopModificat
 {
 	check(MarkersArray);
 
-	const int32 MarkerArrayLength = MarkersArray->WaveCues.Num();
+	const int32 MarkerArrayLength = MarkersArray->CuesAndLoops.Num();
+
+	FSoundWaveCuePoint* SelectedCuePointPtr = nullptr;
 
 	// If there is no active selection, pick first or last elmeent
 	if (MarkersArray->SelectedCue == INDEX_NONE && MarkerArrayLength > 0)
@@ -308,38 +313,56 @@ void FWaveformTransformationMarkerRenderer::CycleMarkerLoopRegion(ELoopModificat
 		{
 		case ELoopModificationControls::SelectNextLoop:
 		{
-			MarkersArray->SelectedCue = MarkersArray->WaveCues[0].CuePointID;
+			MarkersArray->SelectedCue = MarkersArray->CuesAndLoops[0].CuePointID;
+			SelectedCuePointPtr = &MarkersArray->CuesAndLoops[0];
 		}
 			break;
 		case ELoopModificationControls::SelectPreviousLoop:
 		{
-			MarkersArray->SelectedCue = MarkersArray->WaveCues[MarkerArrayLength - 1].CuePointID;
+			MarkersArray->SelectedCue = MarkersArray->CuesAndLoops[MarkerArrayLength - 1].CuePointID;
+			SelectedCuePointPtr = &MarkersArray->CuesAndLoops[MarkerArrayLength - 1];
 		}
 			break;
 		default:
 			break;
 		}
+		return;
 	}
 
 	// Since cues can be removed and added the ids do not match list indexes, so search for matches
 	for (int i = 0; i < MarkerArrayLength; i++)
 	{
-		if (MarkersArray->WaveCues[i].CuePointID == MarkersArray->SelectedCue)
+		if (MarkersArray->CuesAndLoops[i].CuePointID == MarkersArray->SelectedCue)
 		{
 			// Check if we move backwards of forwards
-			if (Modification == ELoopModificationControls::SelectNextLoop && i < MarkerArrayLength - 1)
+			if (Modification == ELoopModificationControls::SelectNextLoop)
 			{
-				MarkersArray->SelectedCue = MarkersArray->WaveCues[i + 1].CuePointID;
+				MarkersArray->SelectedCue = MarkersArray->CuesAndLoops[(i + 1) % MarkerArrayLength].CuePointID;
+				SelectedCuePointPtr = &MarkersArray->CuesAndLoops[(i + 1) % MarkerArrayLength];
 				break;
 			}
-			else if (Modification == ELoopModificationControls::SelectPreviousLoop && i > 0)
+			else if (Modification == ELoopModificationControls::SelectPreviousLoop)
 			{
-				MarkersArray->SelectedCue = MarkersArray->WaveCues[i - 1].CuePointID;
+				MarkersArray->SelectedCue = MarkersArray->CuesAndLoops[(i - 1 + MarkerArrayLength) % MarkerArrayLength].CuePointID;
+				SelectedCuePointPtr = &MarkersArray->CuesAndLoops[(i - 1 + MarkerArrayLength) % MarkerArrayLength];
 				break;
 			}
 		}
 	}
 
+	check(TransformationWaveInfo.SampleRate > 0);
+	if (SelectedCuePointPtr && SelectedCuePointPtr->IsLoopRegion())
+	{
+		// Set Loop preview handles
+		const float StartLoopPosInSeconds = static_cast<float>(SelectedCuePointPtr->FramePosition) / TransformationWaveInfo.SampleRate;
+		const float EndLoopPosInSeconds = static_cast<float>(SelectedCuePointPtr->FramePosition + SelectedCuePointPtr->FrameLength) / TransformationWaveInfo.SampleRate;
+
+		SetIsPreviewingLoopRegion(static_cast<double>(StartLoopPosInSeconds), static_cast<double>(EndLoopPosInSeconds), true);
+	}
+	else
+	{
+		SetIsPreviewingLoopRegion(0.0, -1.0, false);
+	}
 }
 
 FSoundWaveCuePoint* FWaveformTransformationMarkerRenderer::GetSelectedMarker() const
@@ -351,11 +374,11 @@ FSoundWaveCuePoint* FWaveformTransformationMarkerRenderer::GetSelectedMarker() c
 		return nullptr;
 	}
 
-	for (int i = 0; i < MarkersArray->WaveCues.Num(); i++)
+	for (int i = 0; i < MarkersArray->CuesAndLoops.Num(); i++)
 	{
-		if (MarkersArray->WaveCues[i].CuePointID == MarkersArray->SelectedCue)
+		if (MarkersArray->CuesAndLoops[i].CuePointID == MarkersArray->SelectedCue)
 		{
-			return &MarkersArray->WaveCues[i];
+			return &MarkersArray->CuesAndLoops[i];
 		}
 	}
 
@@ -369,7 +392,7 @@ void FWaveformTransformationMarkerRenderer::SelectLoopRegionByKeyboard(const FKe
 	for (int i = 0; i < IndexKeyCodes.Num(); i++)
 	{
 		// Check we have loop regions in this index
-		if (MarkersArray->WaveCues.Num() <= i)
+		if (MarkersArray->CuesAndLoops.Num() <= i)
 		{
 			return;
 		}
@@ -377,6 +400,20 @@ void FWaveformTransformationMarkerRenderer::SelectLoopRegionByKeyboard(const FKe
 		if (PressedKey == IndexKeyCodes[i])
 		{
 			MarkersArray->SelectedCue = i;
+
+			check(TransformationWaveInfo.SampleRate > 0);
+			if (GetSelectedMarker() && GetSelectedMarker()->IsLoopRegion())
+			{
+				// Set Loop preview handles
+				const float StartLoopPosInSeconds = static_cast<float>(GetSelectedMarker()->FramePosition) / TransformationWaveInfo.SampleRate;
+				const float EndLoopPosInSeconds = static_cast<float>(GetSelectedMarker()->FramePosition + GetSelectedMarker()->FrameLength) / TransformationWaveInfo.SampleRate;
+
+				SetIsPreviewingLoopRegion(static_cast<double>(StartLoopPosInSeconds), static_cast<double>(EndLoopPosInSeconds), true);
+			}
+			else
+			{
+				SetIsPreviewingLoopRegion(0.0, -1.0, false);
+			}
 		}
 	}
 }
@@ -396,10 +433,10 @@ void FWaveformTransformationMarkerRenderer::Tick(const FGeometry& AllottedGeomet
 	check(MarkersArray != nullptr);
 
 	TMap<uint32, FSoundWaveCuePoint> CuePointMap;
-	CuePointMap.Reserve(MarkersArray->WaveCues.Num());
+	CuePointMap.Reserve(MarkersArray->CuesAndLoops.Num());
 
 	bool IsMarkerInInteractionRangeValid = false;
-	for (const FSoundWaveCuePoint& Marker : MarkersArray->WaveCues)
+	for (const FSoundWaveCuePoint& Marker : MarkersArray->CuesAndLoops)
 	{
 		//Fix duplicate loop region markers for SoundWaves that have not been reimported 
 		//after import fix
@@ -501,7 +538,7 @@ FReply FWaveformTransformationMarkerRenderer::OnMouseMove(SWidget& OwnerWidget,
 
 	// Evaluate highlight state for loop handles
 	HighlightedCue = INDEX_NONE;
-	for (const FSoundWaveCuePoint& CuePoint : MarkersArray->WaveCues)
+	for (const FSoundWaveCuePoint& CuePoint : MarkersArray->CuesAndLoops)
 	{
 		if (CuePoint.IsLoopRegion() && IsPositionInLoopArea(MousePosition, MyGeometry, &CuePoint))
 		{
@@ -526,7 +563,7 @@ FCursorReply FWaveformTransformationMarkerRenderer::OnCursorQuery(const FGeometr
 	}
 
 	// Set resize on loop handle hover, loops are seperated as we want highlight interactions to superscede region highlights
-	for (const FSoundWaveCuePoint& CuePoint : MarkersArray->WaveCues)
+	for (const FSoundWaveCuePoint& CuePoint : MarkersArray->CuesAndLoops)
 	{
 		const bool bIsLeftHandleHighlighted = IsPositionInLoopHandleArea(MousePosition, MyGeometry, CuePoint, true) || 
 			(MarkerInteractionType == EMarkerInteractionType::ScrubbingMarkerHandleLeft && HighlightedCue == CuePoint.CuePointID);
@@ -544,7 +581,7 @@ FCursorReply FWaveformTransformationMarkerRenderer::OnCursorQuery(const FGeometr
 	}
 
 	// Set grab hand on loop region moving
-	for (const FSoundWaveCuePoint& CuePoint : MarkersArray->WaveCues)
+	for (const FSoundWaveCuePoint& CuePoint : MarkersArray->CuesAndLoops)
 	{
 		if (CuePoint.IsLoopRegion() && IsPositionInLoopArea(MousePosition, MyGeometry, &CuePoint))
 		{
@@ -666,12 +703,13 @@ bool FWaveformTransformationMarkerRenderer::IsPositionInLoopHandleArea(const FVe
 bool FWaveformTransformationMarkerRenderer::SetMarkerInInteractionRange(const FVector2D& InLocalCursorPosition, const FGeometry& WidgetGeometry)
 {
 	check(TransformationWaveInfo.NumChannels > 0);
+	check(TransformationWaveInfo.SampleRate > 0);
 	ensure(LoopHandlePixelWidth > 0);
 	check(MarkersArray != nullptr);
 
-	InteractionRanges.Reset(MarkersArray->WaveCues.Num());
+	InteractionRanges.Reset(MarkersArray->CuesAndLoops.Num());
 	
-	for (FSoundWaveCuePoint& CuePoint : MarkersArray->WaveCues)
+	for (FSoundWaveCuePoint& CuePoint : MarkersArray->CuesAndLoops)
 	{
 		if (IsPositionInInteractionRange(CuePoint.FramePosition, InLocalCursorPosition, WidgetGeometry))
 		{
@@ -698,7 +736,7 @@ bool FWaveformTransformationMarkerRenderer::SetMarkerInInteractionRange(const FV
 	check(PixelsPerFrame > 0);
 
 	// Evaluate loop area if no handles are selected
-	for (FSoundWaveCuePoint& CuePoint : MarkersArray->WaveCues)
+	for (FSoundWaveCuePoint& CuePoint : MarkersArray->CuesAndLoops)
 	{
 		if (CuePoint.FrameLength > 0 && IsPositionInLoopArea(InLocalCursorPosition, WidgetGeometry, &CuePoint))
 		{
@@ -750,6 +788,20 @@ void FWaveformTransformationMarkerRenderer::SetActiveLoopRegion(int32 CuePointID
 	MarkersArray->SelectedCue = CuePointID;
 	EndTransaction();
 	MarkersHandle->NotifyPostChange(EPropertyChangeType::ValueSet);
+
+	// Revert Loop preview handles
+	if (CuePointID == INDEX_NONE)
+	{
+		check(TransformationWaveInfo.SampleRate > 0);
+		SetIsPreviewingLoopRegion(0.0, -1.0, false);
+	}
+}
+
+void FWaveformTransformationMarkerRenderer::SetIsPreviewingLoopRegion(double InStartTime, double InEndTime, bool bIsPreviewing)
+{
+	IsPreviewingLoopRegionHandle->SetValue(bIsPreviewing);
+	StartLoopTimeHandle->SetValue(InStartTime);
+	EndLoopTimeHandle->SetValue(InEndTime);
 }
 
 FVector2D FWaveformTransformationMarkerRenderer::GetLocalCursorPosition(const FPointerEvent& MouseEvent, const FGeometry& EventGeometry)
@@ -779,7 +831,7 @@ FWaveformTransformationMarkerRenderer::EMarkerInteractionType FWaveformTransform
 		return EMarkerInteractionType::None;
 	}
 
-	for (FSoundWaveCuePoint CuePoint : MarkersArray->WaveCues)
+	for (FSoundWaveCuePoint CuePoint : MarkersArray->CuesAndLoops)
 	{
 		if (!CuePoint.IsLoopRegion())
 		{
diff --git a/Engine/Plugins/Editor/WaveformEditor/Source/WaveformTransformationsWidgets/Private/WaveformTransformationMarkersRenderer.h b/Engine/Plugins/Editor/WaveformEditor/Source/WaveformTransformationsWidgets/Private/WaveformTransformationMarkersRenderer.h
--- a/Engine/Plugins/Editor/WaveformEditor/Source/WaveformTransformationsWidgets/Private/WaveformTransformationMarkersRenderer.h
+++ b/Engine/Plugins/Editor/WaveformEditor/Source/WaveformTransformationsWidgets/Private/WaveformTransformationMarkersRenderer.h
@@ -42,6 +42,7 @@ private:
 	bool SetMarkerInInteractionRange(const FVector2D& InLocalCursorPosition, const FGeometry& WidgetGeometry);
 	bool SelectLoopArea(const FVector2D& InLocalCursorPosition, const FGeometry& WidgetGeometry);
 	void SetActiveLoopRegion(int32 CuePointID);
+	void SetIsPreviewingLoopRegion(double InStartTime, double InEndTime, bool bIsPreviewing);
 	static FVector2D GetLocalCursorPosition(const FPointerEvent& MouseEvent, const FGeometry& EventGeometry);
 
 	TArray<TRange<float>> InteractionRanges;
@@ -98,6 +99,9 @@ private:
 	float LabelFontSize = WaveformTransformationWidgetSharedDefaults::DefaultLabelFontSize;
 
 	TSharedPtr<IPropertyHandle> MarkersHandle;
+	TSharedPtr<IPropertyHandle> StartLoopTimeHandle;
+	TSharedPtr<IPropertyHandle> EndLoopTimeHandle;
+	TSharedPtr<IPropertyHandle> IsPreviewingLoopRegionHandle;
 
 	TWeakPtr<SWindow> CreateLoopMenuWindow;
 	const TArray<int32> SizeIncrements = { 1, 10, 50, 100 };
diff --git a/Engine/Source/Runtime/AudioExtensions/Public/IWaveformTransformation.h b/Engine/Source/Runtime/AudioExtensions/Public/IWaveformTransformation.h
--- a/Engine/Source/Runtime/AudioExtensions/Public/IWaveformTransformation.h
+++ b/Engine/Source/Runtime/AudioExtensions/Public/IWaveformTransformation.h
@@ -28,6 +28,13 @@ namespace Audio
 		uint32 NumEditedSamples = 0;
 	};
 
+	enum class ETransformationPriority : uint8
+	{
+		None = 0,
+		Low,
+		High
+	};
+
 	/*
 	 * Base class for the object that processes waveform data
 	 * Pass tweakable variables from its paired settings UObject in the constructor in UWaveformTransformationBase::CreateTransformation
@@ -42,7 +49,7 @@ namespace Audio
 		virtual void ProcessAudio(FWaveformTransformationWaveInfo& InOutWaveInfo) const {};
 		
 		virtual bool SupportsRealtimePreview() const { return false; }
-		virtual bool CanChangeFileLength() const { return false; }
+		virtual constexpr ETransformationPriority FileChangeLengthPriority() const { return ETransformationPriority::None; }
 		virtual bool CanChangeChannelCount() const { return false; }
 
 		virtual ~IWaveTransformation() {};
@@ -102,6 +109,8 @@ struct FWaveTransformUObjectConfiguration
 	float StartTime = 0.f; 
 	float EndTime = -1.f; 
 	TArray<FSoundWaveCuePoint> WaveCues; // List of cues parsed from the wave file
+	bool bIsPreviewingLoopRegion = false;
+	bool bCachedSoundWaveLoopState = false;
 };
 
 #if WITH_EDITORONLY_DATA
@@ -122,6 +131,25 @@ public:
 	virtual Audio::FTransformationPtr CreateTransformation() const { return nullptr; }
 	virtual void UpdateConfiguration(FWaveTransformUObjectConfiguration& InOutConfiguration) {};
 	virtual void OverwriteTransformation() {};
+	virtual constexpr Audio::ETransformationPriority GetTransformationPriority() const { return Audio::ETransformationPriority::None; }
+
+	// Sort to ensure proper order of operation for audio processing
+	static void SortTransformationsArray(TArray<TObjectPtr<UWaveformTransformationBase>>& InOutTransformations) 
+	{
+		Algo::Sort(InOutTransformations, [](const TObjectPtr<UWaveformTransformationBase>& TransformationA, const TObjectPtr<UWaveformTransformationBase>& TransformationB)
+			{
+				if ((!TransformationA && TransformationB) || (!TransformationA && !TransformationB))
+				{
+					return false;
+				}
+				else if (TransformationA && !TransformationB)
+				{
+					return true;
+				}
+
+				return TransformationA->GetTransformationPriority() > TransformationB->GetTransformationPriority();
+			});
+	}
 	
 #if WITH_EDITOR
 	virtual void OverwriteSoundWaveData(USoundWave& InOutSoundWave) {};
diff --git a/Engine/Source/Runtime/Engine/Private/SoundWave.cpp b/Engine/Source/Runtime/Engine/Private/SoundWave.cpp
--- a/Engine/Source/Runtime/Engine/Private/SoundWave.cpp
+++ b/Engine/Source/Runtime/Engine/Private/SoundWave.cpp
@@ -4704,6 +4704,34 @@ const FWaveTransformUObjectConfiguration& USoundWave::UpdateTransformations()
 	TransformationChainConfig.StartTime = 0.f;
 	TransformationChainConfig.EndTime = TotalSamples / TransformationChainConfig.SampleRate;
 	TransformationChainConfig.WaveCues = CuePoints;
+	TransformationChainConfig.bIsPreviewingLoopRegion = false;
+	TransformationChainConfig.bCachedSoundWaveLoopState = static_cast<bool>(bLooping);
+
+	TSet<UClass*> TransformationClassSet;
+	for (int32 i = 0; i < Transformations.Num(); ++i)
+	{
+		TObjectPtr<UWaveformTransformationBase>& Transformation = Transformations[i];
+		
+		if (Transformation == nullptr)
+		{
+			continue;
+		}
+
+		UClass* TransformationClass = Transformation->GetClass();
+
+		if (TransformationClassSet.Contains(TransformationClass))
+		{
+			UE_LOG(LogAudio, Warning, TEXT("Can't have more than one %s on a SoundWave"), *TransformationClass->GetName());
+			Transformation = nullptr;
+
+			break;
+		}
+
+		TransformationClassSet.Add(Transformation->GetClass());
+	}
+
+	//Sort by priority to have a proper order of operations and simplify transformation audio processing
+	UWaveformTransformationBase::SortTransformationsArray(Transformations);
 
 	for (TObjectPtr<UWaveformTransformationBase>& Transformation : Transformations)
 	{
@@ -4713,6 +4741,15 @@ const FWaveTransformUObjectConfiguration& USoundWave::UpdateTransformations()
 		}
 	}
 
+	if (TransformationChainConfig.bIsPreviewingLoopRegion)
+	{
+		bLooping = 1;
+	}
+	else
+	{
+		bLooping = static_cast<uint8>(TransformationChainConfig.bCachedSoundWaveLoopState);
+	}
+
 	return TransformationChainConfig;
 }
 
