https://jira.it.epicgames.com/browse/UE-166248
To ensure the waveform editor's playstate is in sync with the SoundWave on implicit state changes when editing transformations, clear the AudioComponent's paused state to keep it in sync with the USoundWave.

https://jira.it.epicgames.com/browse/UE-229205
Skip the Playhead to markers support for waveform editor.

diff --git a/Engine/Plugins/Editor/WaveformEditor/Source/WaveformEditor/Private/WaveformEditor.cpp b/Engine/Plugins/Editor/WaveformEditor/Source/WaveformEditor/Private/WaveformEditor.cpp
--- a/Engine/Plugins/Editor/WaveformEditor/Source/WaveformEditor/Private/WaveformEditor.cpp
+++ b/Engine/Plugins/Editor/WaveformEditor/Source/WaveformEditor/Private/WaveformEditor.cpp
@@ -404,6 +404,51 @@ void FWaveformEditor::DeleteMarker()
 	}
 }
 
+void FWaveformEditor::SkipToNextMarker()
+{
+	check(SoundWave != nullptr);
+
+	if (TObjectPtr<UWaveformTransformationMarkers> MarkerTransformation = GetOrAddMarkerTransformation())
+	{
+		check(TransportCoordinator != nullptr);
+		check(SoundWave->TotalSamples > 0);
+		const float PlaybackFramePosition = TransportCoordinator->GetFocusPoint() * SoundWave->TotalSamples;
+		float NewPlaybackFramePosition = FLT_MAX;
+
+		// Find the nearest marker beyond the current playhead.
+		check(MarkerTransformation->Markers != nullptr);
+		for (int i = 0; i < MarkerTransformation->Markers->CuesAndLoops.Num(); i++)
+		{
+			if (MarkerTransformation->Markers->CuesAndLoops[i].FramePosition > PlaybackFramePosition + 1)
+			{
+				if (MarkerTransformation->Markers->CuesAndLoops[i].FramePosition < NewPlaybackFramePosition)
+				{
+					NewPlaybackFramePosition = MarkerTransformation->Markers->CuesAndLoops[i].FramePosition;
+				}
+			}
+		}
+
+		if (NewPlaybackFramePosition != FLT_MAX)
+		{
+			float ActiveDuration = TransformationChainConfig.EndTime - TransformationChainConfig.StartTime;
+			if (ActiveDuration <= 0.0f)
+			{
+				// If no active+initialized transformation, ActiveDuration will be <= 0.0f, so fallback to sound duration:
+				ActiveDuration = SoundWave->Duration; // not GetDuration() to get the raw duration and not INDEFINITELY_LOOPING_DURATION if looping
+			}
+			if (ActiveDuration > 0)
+			{
+				check(TransformationChainConfig.SampleRate > 0);
+				const float ActiveDurationInFrames = ActiveDuration * TransformationChainConfig.SampleRate;
+				const float AdjustedNewPlaybackFramePosition = FMath::Max(NewPlaybackFramePosition - TransformationChainConfig.StartTime * TransformationChainConfig.SampleRate, 0.0f);
+				
+				HandlePlayheadScrub(AdjustedNewPlaybackFramePosition / ActiveDurationInFrames, false);
+				TransportCoordinator->SetProgressRatio(AdjustedNewPlaybackFramePosition / ActiveDurationInFrames);
+			}
+		}
+	}
+}
+
 void FWaveformEditor::ModifyMarkerLoopRegion(ELoopModificationControls Modification)
 {
 	check(SoundWave);
@@ -1094,6 +1139,10 @@ bool FWaveformEditor::BindCommands()
 		Commands.DeleteMarker,
 		FExecuteAction::CreateSP(this, &FWaveformEditor::DeleteMarker));
 
+	ToolkitCommands->MapAction(
+		Commands.SkipToNextMarker,
+		FExecuteAction::CreateSP(this, &FWaveformEditor::SkipToNextMarker));
+
 	return true;
 }
 
@@ -1448,14 +1497,19 @@ void FWaveformEditor::HandleAudioComponentPlayStateChanged(const UAudioComponent
 	default:
 		break;
 	case EAudioComponentPlayState::Stopped:
-		const bool bIgnoreStopState = TransportCoordinator->IsScrubbing() || bIsInteractingWithTransformations;
-
-		if (!bIgnoreStopState)
+		if (!TransportCoordinator->IsScrubbing())
 		{
-			const float CachedAudioStartTimeAsPercentage = TransportController->GetCachedAudioStartTimeAsPercentage();
-			TransportCoordinator->SetProgressRatio(CachedAudioStartTimeAsPercentage); // show cached start time for next play rather than playhead stop point (doesn't affect pause)
-		}
+			// USoundWave::PostEditChangeProperty calls USoundWave::UpdateAsset which frees the soundwave's resources,
+			// stopping the sound, calling this handler.  Avoiding this implicit state change is currently too risky and
+			// so we clear the AudioComponent's paused state instead to keep it in sync with the USoundWave.
+			AudioComponent->SetPaused(false);
 
+			if (!bIsInteractingWithTransformations)
+			{
+				const float CachedAudioStartTimeAsPercentage = TransportController->GetCachedAudioStartTimeAsPercentage();
+				TransportCoordinator->SetProgressRatio(CachedAudioStartTimeAsPercentage); // show cached start time for next play rather than playhead stop point (doesn't affect pause)
+			}
+		}
 		break;
 	}
 }
diff --git a/Engine/Plugins/Editor/WaveformEditor/Source/WaveformEditor/Private/WaveformEditorCommands.cpp b/Engine/Plugins/Editor/WaveformEditor/Source/WaveformEditor/Private/WaveformEditorCommands.cpp
--- a/Engine/Plugins/Editor/WaveformEditor/Source/WaveformEditor/Private/WaveformEditorCommands.cpp
+++ b/Engine/Plugins/Editor/WaveformEditor/Source/WaveformEditor/Private/WaveformEditorCommands.cpp
@@ -57,6 +57,8 @@ void FWaveformEditorCommands::RegisterCommands()
 	UI_COMMAND(CreateMarker, "Create Marker", "Toggle adding Marker Cue Transformation onto the sound wave", EUserInterfaceActionType::Button, FInputChord());
 	UI_COMMAND(CreateLoopRegion, "Create Loop Region", "Toggle adding Marker Loop Region Transformation onto the sound wave", EUserInterfaceActionType::Button, FInputChord());
 	UI_COMMAND(DeleteMarker, "Delete Marker Or Loop Region", "Delete a selected Marker Cue or Loop Region from the sound wave", EUserInterfaceActionType::Button, FInputChord());
+
+	UI_COMMAND(SkipToNextMarker, "Skip To Next Marker", "Skips the playhead to the next marker", EUserInterfaceActionType::Button, FInputChord(EKeys::G));
 }
 	
 const FWaveformEditorCommands& FWaveformEditorCommands::Get()
diff --git a/Engine/Plugins/Editor/WaveformEditor/Source/WaveformEditor/Public/WaveformEditor.h b/Engine/Plugins/Editor/WaveformEditor/Source/WaveformEditor/Public/WaveformEditor.h
--- a/Engine/Plugins/Editor/WaveformEditor/Source/WaveformEditor/Public/WaveformEditor.h
+++ b/Engine/Plugins/Editor/WaveformEditor/Source/WaveformEditor/Public/WaveformEditor.h
@@ -114,6 +114,7 @@ private:
 	bool CanFadeOut();
 	void CreateMarker(bool bIsLoopRegion);
 	void DeleteMarker();
+	void SkipToNextMarker(); // Skips the playhead to the next marker
 
 	void ModifyMarkerLoopRegion(ELoopModificationControls Modification);
 	void CycleMarkerLoopRegion(ELoopModificationControls Modification);
diff --git a/Engine/Plugins/Editor/WaveformEditor/Source/WaveformEditor/Public/WaveformEditorCommands.h b/Engine/Plugins/Editor/WaveformEditor/Source/WaveformEditor/Public/WaveformEditorCommands.h
--- a/Engine/Plugins/Editor/WaveformEditor/Source/WaveformEditor/Public/WaveformEditorCommands.h
+++ b/Engine/Plugins/Editor/WaveformEditor/Source/WaveformEditor/Public/WaveformEditorCommands.h
@@ -56,4 +56,6 @@ public:
 	TSharedPtr<FUICommandInfo> CreateMarker;
 	TSharedPtr<FUICommandInfo> CreateLoopRegion;
 	TSharedPtr<FUICommandInfo> DeleteMarker;
+
+	TSharedPtr<FUICommandInfo> SkipToNextMarker;
 };
