Waveform Editor: Undo functionality fixes

diff --git a/Engine/Plugins/Editor/WaveformEditor/Source/WaveformEditor/Private/WaveformEditor.cpp b/Engine/Plugins/Editor/WaveformEditor/Source/WaveformEditor/Private/WaveformEditor.cpp
--- a/Engine/Plugins/Editor/WaveformEditor/Source/WaveformEditor/Private/WaveformEditor.cpp
+++ b/Engine/Plugins/Editor/WaveformEditor/Source/WaveformEditor/Private/WaveformEditor.cpp
@@ -275,11 +275,14 @@ TObjectPtr<UWaveformTransformationMarkers> FWaveformEditor::GetOrAddMarkerTransf
 void FWaveformEditor::ToggleFadeIn()
 {
 	check(SoundWave);
+	check(GEditor != nullptr);
+	GEditor->BeginTransaction(TEXT("PropertyEditor"), LOCTEXT("SetNumericPropertyTransaction", "ToggleFadeIn"), nullptr);
+	SoundWave->Modify();
 
 	TObjectPtr<UWaveformTransformationTrimFade> TrimFadeTransformation = GetOrAddTrimFadeTransformation();
 	check(TrimFadeTransformation);
 
-	SoundWave->Modify();
+	TrimFadeTransformation->Modify();
 	
 	//if fade duration is 0, assume there is no fade present and toggle on the fade with the curve data selected in the menu
 	if (TrimFadeTransformation->StartFadeTime == 0.0f)
@@ -304,6 +307,8 @@ void FWaveformEditor::ToggleFadeIn()
 	}
 
 	NotifyPostTransformationChange(EPropertyChangeType::ValueSet);
+
+	GEditor->EndTransaction();
 }
 
 bool FWaveformEditor::CanFadeIn()
@@ -314,11 +319,14 @@ bool FWaveformEditor::CanFadeIn()
 void FWaveformEditor::ToggleFadeOut()
 {
 	check(SoundWave);
+	check(GEditor != nullptr);
+	GEditor->BeginTransaction(TEXT("PropertyEditor"), LOCTEXT("SetNumericPropertyTransaction", "ToggleFadeOut"), nullptr);
+	SoundWave->Modify();
 
 	TObjectPtr<UWaveformTransformationTrimFade> TrimFadeTransformation = GetOrAddTrimFadeTransformation();
 	check(TrimFadeTransformation);
 
-	SoundWave->Modify();
+	TrimFadeTransformation->Modify();
 
 	//if fade duration is 0, assume there is no fade present and toggle on the fade with the curve data selected in the menu
 	if (TrimFadeTransformation->EndFadeTime == 0.0f)
@@ -343,6 +351,7 @@ void FWaveformEditor::ToggleFadeOut()
 	}
 
 	NotifyPostTransformationChange(EPropertyChangeType::ValueSet);
+	GEditor->EndTransaction();
 }
 
 bool FWaveformEditor::CanFadeOut()
@@ -352,12 +361,17 @@ bool FWaveformEditor::CanFadeOut()
 
 void FWaveformEditor::CreateMarker(bool bIsLoopRegion)
 {
+	check(SoundWave);
+	check(GEditor != nullptr);
+	GEditor->BeginTransaction(TEXT("PropertyEditor"), LOCTEXT("SetNumericPropertyTransaction", "CreateMarker"), nullptr);
+	SoundWave->Modify();
+
 	TObjectPtr<UWaveformTransformationMarkers> MarkerTransformation = GetOrAddMarkerTransformation();
 	check(MarkerTransformation);
 
 	int32 HighestCueID = INDEX_NONE;
 
-	SoundWave->Modify();
+	MarkerTransformation->Markers->Modify();
 
 	for (FSoundWaveCuePoint CuePoint : MarkerTransformation->Markers->CuesAndLoops)
 	{
@@ -379,17 +393,26 @@ void FWaveformEditor::CreateMarker(bool bIsLoopRegion)
 
 	MarkerTransformation->Markers->CuesAndLoops.Add(NewCuePoint);
 	NotifyPostTransformationChange(EPropertyChangeType::ValueSet);
+	GEditor->EndTransaction();
 }
 
 void FWaveformEditor::DeleteMarker()
 {
+	check(SoundWave);
+	check(GEditor != nullptr);
+	GEditor->BeginTransaction(TEXT("PropertyEditor"), LOCTEXT("SetNumericPropertyTransaction", "DeleteMarker"), nullptr);
+	SoundWave->Modify();
+
 	TObjectPtr<UWaveformTransformationMarkers> MarkerTransformation = GetOrAddMarkerTransformation();
 	check(MarkerTransformation);
 
 	int32 CueToDelete = MarkerTransformation->Markers->SelectedCue;
 
+	MarkerTransformation->Markers->Modify();
+
 	if (CueToDelete == INDEX_NONE)
 	{
+		GEditor->EndTransaction();
 		return; // No cue to selected
 	}
 
@@ -402,6 +425,10 @@ void FWaveformEditor::DeleteMarker()
 			break;
 		}
 	}
+
+	NotifyPostTransformationChange(EPropertyChangeType::ValueSet);
+
+	GEditor->EndTransaction();
 }
 
 void FWaveformEditor::SkipToNextMarker()
diff --git a/Engine/Plugins/Editor/WaveformEditor/Source/WaveformTransformationsWidgets/Private/WaveformTransformationMarkersRenderer.cpp b/Engine/Plugins/Editor/WaveformEditor/Source/WaveformTransformationsWidgets/Private/WaveformTransformationMarkersRenderer.cpp
--- a/Engine/Plugins/Editor/WaveformEditor/Source/WaveformTransformationsWidgets/Private/WaveformTransformationMarkersRenderer.cpp
+++ b/Engine/Plugins/Editor/WaveformEditor/Source/WaveformTransformationsWidgets/Private/WaveformTransformationMarkersRenderer.cpp
@@ -244,6 +244,7 @@ void FWaveformTransformationMarkerRenderer::ModifyMarkerLoopRegion(ELoopModifica
 
 	int32 FramesToShift = (SizeIncrements[SelectedIncrement] / 1000.0f) * TransformationWaveInfo.SampleRate;
 	int32 MaxFrames = TransformationWaveInfo.NumAvilableSamples;
+	BeginTransaction(TEXT("PropertyEditor"), LOCTEXT("SetNumericPropertyTransaction", "Modifying Marker"), nullptr);
 
 	MarkersArray->Modify();
 
@@ -487,7 +488,6 @@ FReply FWaveformTransformationMarkerRenderer::OnMouseButtonDown(SWidget& OwnerWi
 	if (MarkerInteractionType != EMarkerInteractionType::None)
 	{
 		MarkersHandle->NotifyPreChange();
-		BeginTransaction(TEXT("PropertyEditor"), FText::Format(NSLOCTEXT("PropertyEditor", "SetNumericPropertyTransaction", "Edit {0}"), FText::FromName(GET_MEMBER_NAME_CHECKED(UWaveformTransformationMarkers, Markers))), nullptr);
 
 		return FReply::Handled().CaptureMouse(OwnerWidget.AsShared()).PreventThrottling();
 	}
@@ -511,10 +511,17 @@ FReply FWaveformTransformationMarkerRenderer::OnMouseButtonUp(SWidget& OwnerWidg
 	{
 		MarkersArray->Modify();
 		SetPropertyValueDependingOnInteractionType(MouseEvent, MyGeometry);
-		EndTransaction();
-		MarkerInteractionType = EMarkerInteractionType::None;
 		
 		MarkersHandle->NotifyPostChange(EPropertyChangeType::ValueSet);
+
+		// If we know it was an interaction that started a transaction, end the transaction
+		if (MarkerInteractionType == EMarkerInteractionType::ScrubbingMarkerHandleLeft || MarkerInteractionType == EMarkerInteractionType::ScrubbingMarkerHandleRight
+			|| MarkerInteractionType == EMarkerInteractionType::LoopHandle || MarkerInteractionType == EMarkerInteractionType::MarkerHandle)
+		{
+			EndTransaction();
+		}
+
+		MarkerInteractionType = EMarkerInteractionType::None;
 		
 		return FReply::Handled().ReleaseMouseCapture();
 	}
@@ -841,24 +848,32 @@ FWaveformTransformationMarkerRenderer::EMarkerInteractionType FWaveformTransform
 		if (IsPositionInLoopHandleArea(InLocalCursorPosition, WidgetGeometry, CuePoint, true))
 		{
 			SetMarkerInInteractionRange(InLocalCursorPosition, WidgetGeometry);
+			BeginTransaction(TEXT("PropertyEditor"), LOCTEXT("ResizingMarker", "ResizingMarker"), nullptr);
+			MarkersArray->Modify();
 			return EMarkerInteractionType::ScrubbingMarkerHandleLeft;
 		}
 
 		if (IsPositionInLoopHandleArea(InLocalCursorPosition, WidgetGeometry, CuePoint, false))
 		{
 			SetMarkerInInteractionRange(InLocalCursorPosition, WidgetGeometry);
+			BeginTransaction(TEXT("PropertyEditor"), LOCTEXT("ResizingMarker", "ResizingMarker"), nullptr);
+			MarkersArray->Modify();
 			return EMarkerInteractionType::ScrubbingMarkerHandleRight;
 		}
 	}
 
 	if (SetMarkerInInteractionRange(InLocalCursorPosition, WidgetGeometry))
 	{
+		BeginTransaction(TEXT("PropertyEditor"), LOCTEXT("MoveMarker", "MoveMarker"), nullptr);
+		MarkersArray->Modify();
 		return EMarkerInteractionType::MarkerHandle;
 	}
 	
 	if (SelectLoopArea(InLocalCursorPosition, WidgetGeometry))
 	{
 		SetMarkerInInteractionRange(InLocalCursorPosition, WidgetGeometry);
+		BeginTransaction(TEXT("PropertyEditor"), LOCTEXT("MoveLoop", "MoveLoop"), nullptr);
+		MarkersArray->Modify();
 		return EMarkerInteractionType::LoopHandle;
 	}
 
@@ -868,9 +883,11 @@ FWaveformTransformationMarkerRenderer::EMarkerInteractionType FWaveformTransform
 void FWaveformTransformationMarkerRenderer::SetPropertyValueDependingOnInteractionType(const FPointerEvent& MouseEvent, const FGeometry& WidgetGeometry)
 {
 	ensure(PixelsPerFrame > 0);
+	check(MarkersHandle);
 
 	const FVector2D LocalCursorPosition = GetLocalCursorPosition(MouseEvent, WidgetGeometry);
 	const float LocalCursorXClamped = FMath::Clamp(LocalCursorPosition.X, 0.f, WidgetGeometry.GetLocalSize().X);
+	MarkersHandle->NotifyPreChange();
 
 	switch (MarkerInteractionType)
 	{
@@ -939,9 +956,14 @@ void FWaveformTransformationMarkerRenderer::ShowConvertToLoopAtMarker(const FGeo
 				.Text(LOCTEXT("WaveformTransformationMarkerRendererConvertLoop","Convert To Loop Region"))
 				.OnClicked_Lambda([this, CueToModify]() -> FReply
 					{
+						BeginTransaction(TEXT("PropertyEditor"), LOCTEXT("ConvertMarkerToLoop", "ConvertMarkerToLoop"), nullptr);
+
+						MarkersArray->Modify();
 						MarkersArray->EnableLoopRegion(CueToModify);
 						CueToModify->FrameLength = MarkerInInteractionRange.GetMinLoopSize();
 
+						EndTransaction();
+
 						if (TSharedPtr<SWindow> LockedFadeModeMenuWindow = CreateLoopMenuWindow.Pin())
 						{
 							FSlateApplication::Get().RequestDestroyWindow(LockedFadeModeMenuWindow.ToSharedRef());
