Waveform Editor: Realtime loop preview audio for UI

diff --git a/Engine/Plugins/Editor/WaveformEditor/Source/WaveformEditor/Private/WaveformEditor.cpp b/Engine/Plugins/Editor/WaveformEditor/Source/WaveformEditor/Private/WaveformEditor.cpp
--- a/Engine/Plugins/Editor/WaveformEditor/Source/WaveformEditor/Private/WaveformEditor.cpp
+++ b/Engine/Plugins/Editor/WaveformEditor/Source/WaveformEditor/Private/WaveformEditor.cpp
@@ -1225,6 +1225,12 @@ void FWaveformEditor::OnAssetReimport(UObject* ReimportedObject, bool bSuccessfu
 }
 
 
+void FWaveformEditor::NotifyPreChange(FEditPropertyChain* PropertyAboutToChange)
+{
+	check(TransportController != nullptr);
+	bWasPlayingBeforeChange = TransportController->IsPlaying();
+}
+
 void FWaveformEditor::NotifyPostChange(const FPropertyChangedEvent& PropertyChangedEvent, class FEditPropertyChain* PropertyThatChanged)
 {	
 	TDoubleLinkedList<FProperty*>::TDoubleLinkedListNode* PropertyNode = PropertyThatChanged->GetActiveMemberNode();
@@ -1257,6 +1263,7 @@ void FWaveformEditor::NotifyPostChange(const FPropertyChangedEvent& PropertyChan
 			PlaybackTimeBeforeTransformInteraction = LastReceivedPlaybackPercent * (TransformationChainConfig.EndTime - TransformationChainConfig.StartTime);
 			StartTimeBeforeTransformInteraction = TransformationChainConfig.StartTime;
 			AudioComponent->Stop();
+			bWasPlayingBeforeChange = false;
 		}
 	}
 
@@ -1306,6 +1313,11 @@ void FWaveformEditor::NotifyPostChange(const FPropertyChangedEvent& PropertyChan
 			default:
 				break;
 			}
+
+			if (bWasPlayingBeforeChange)
+			{
+				AudioComponent->Play();
+			}
 		}
 
 		bIsInteractingWithTransformations = false;
diff --git a/Engine/Plugins/Editor/WaveformEditor/Source/WaveformEditor/Public/WaveformEditor.h b/Engine/Plugins/Editor/WaveformEditor/Source/WaveformEditor/Public/WaveformEditor.h
--- a/Engine/Plugins/Editor/WaveformEditor/Source/WaveformEditor/Public/WaveformEditor.h
+++ b/Engine/Plugins/Editor/WaveformEditor/Source/WaveformEditor/Public/WaveformEditor.h
@@ -42,7 +42,7 @@ public:
 	void OnAssetReimport(UObject* ReimportedObject, bool bSuccessfullReimport);
 
 	/** FNotifyHook interface */
-	void NotifyPreChange(class FEditPropertyChain* PropertyAboutToChange) override {};
+	void NotifyPreChange(class FEditPropertyChain* PropertyAboutToChange) override;
 	void NotifyPostChange(const FPropertyChangedEvent& PropertyChangedEvent, class FEditPropertyChain* PropertyThatChanged) override;
 
 	/** FEditorUndo interface */
@@ -152,6 +152,7 @@ private:
 	TObjectPtr<UAudioComponent> AudioComponent = nullptr;
 	bool bWasPlayingBeforeScrubbing = false;
 	bool bIsInteractingWithTransformations = false;
+	bool bWasPlayingBeforeChange = false;
 	float LastReceivedPlaybackPercent = 0.f;
 	EAudioComponentPlayState TransformInteractionPlayState = EAudioComponentPlayState::Stopped;
 	float PlaybackTimeBeforeTransformInteraction = 0.f;
diff --git a/Engine/Plugins/Editor/WaveformEditor/Source/WaveformTransformationsWidgets/Private/WaveformTransformationMarkersRenderer.cpp b/Engine/Plugins/Editor/WaveformEditor/Source/WaveformTransformationsWidgets/Private/WaveformTransformationMarkersRenderer.cpp
--- a/Engine/Plugins/Editor/WaveformEditor/Source/WaveformTransformationsWidgets/Private/WaveformTransformationMarkersRenderer.cpp
+++ b/Engine/Plugins/Editor/WaveformEditor/Source/WaveformTransformationsWidgets/Private/WaveformTransformationMarkersRenderer.cpp
@@ -296,6 +296,8 @@ void FWaveformTransformationMarkerRenderer::ModifyMarkerLoopRegion(ELoopModifica
 	LoopRegion->FramePosition = FMath::Clamp(LoopRegion->FramePosition, 0, LoopRegion->FramePosition);
 	LoopRegion->FrameLength = FMath::Clamp(LoopRegion->FrameLength, 0, MaxFrames - LoopRegion->FramePosition);
 
+	PreviewSelectedLoop();
+
 	EndTransaction();
 }
 
@@ -354,11 +356,7 @@ void FWaveformTransformationMarkerRenderer::CycleMarkerLoopRegion(ELoopModificat
 	check(TransformationWaveInfo.SampleRate > 0);
 	if (SelectedCuePointPtr && SelectedCuePointPtr->IsLoopRegion())
 	{
-		// Set Loop preview handles
-		const float StartLoopPosInSeconds = static_cast<float>(SelectedCuePointPtr->FramePosition) / TransformationWaveInfo.SampleRate;
-		const float EndLoopPosInSeconds = static_cast<float>(SelectedCuePointPtr->FramePosition + SelectedCuePointPtr->FrameLength) / TransformationWaveInfo.SampleRate;
-
-		SetIsPreviewingLoopRegion(static_cast<double>(StartLoopPosInSeconds), static_cast<double>(EndLoopPosInSeconds), true);
+		PreviewSelectedLoop();
 	}
 	else
 	{
@@ -405,11 +403,7 @@ void FWaveformTransformationMarkerRenderer::SelectLoopRegionByKeyboard(const FKe
 			check(TransformationWaveInfo.SampleRate > 0);
 			if (GetSelectedMarker() && GetSelectedMarker()->IsLoopRegion())
 			{
-				// Set Loop preview handles
-				const float StartLoopPosInSeconds = static_cast<float>(GetSelectedMarker()->FramePosition) / TransformationWaveInfo.SampleRate;
-				const float EndLoopPosInSeconds = static_cast<float>(GetSelectedMarker()->FramePosition + GetSelectedMarker()->FrameLength) / TransformationWaveInfo.SampleRate;
-
-				SetIsPreviewingLoopRegion(static_cast<double>(StartLoopPosInSeconds), static_cast<double>(EndLoopPosInSeconds), true);
+				PreviewSelectedLoop();
 			}
 			else
 			{
@@ -497,6 +491,16 @@ FReply FWaveformTransformationMarkerRenderer::OnMouseButtonDown(SWidget& OwnerWi
 
 FReply FWaveformTransformationMarkerRenderer::OnMouseButtonUp(SWidget& OwnerWidget, const FGeometry& MyGeometry, const FPointerEvent& MouseEvent)
 {
+	if (MarkersArray->SelectedCue == INDEX_NONE)
+	{
+		SetIsPreviewingLoopRegion(0.0, -1.0, false);
+	}
+	else
+	{
+		PreviewSelectedLoop();
+	}
+
+
 	if (TSharedPtr<SWindow> LockedFadeModeMenuWindow = CreateLoopMenuWindow.Pin())
 	{
 		FSlateApplication::Get().RequestDestroyWindow(LockedFadeModeMenuWindow.ToSharedRef());
@@ -694,7 +698,7 @@ bool FWaveformTransformationMarkerRenderer::IsPositionInLoopHandleArea(const FVe
 	const uint32 StartFrameOffset = TransformationWaveInfo.StartFrameOffset / TransformationWaveInfo.NumChannels;
 	float LoopRegionPosition = IsLeftHandle ?
 		(CueMarker.FramePosition - StartFrameOffset) * PixelsPerFrame :
-		((CueMarker.FramePosition + CueMarker.FrameLength) * PixelsPerFrame) - LoopHandlePixelWidth;
+		((CueMarker.FramePosition - StartFrameOffset + CueMarker.FrameLength) * PixelsPerFrame) - LoopHandlePixelWidth;
 
 	// Check that only the handle within the loop area is highlighted
 	if (InLocalCursorPosition.X > LoopRegionPosition && InLocalCursorPosition.X < LoopRegionPosition + LoopHandlePixelWidth &&
@@ -748,7 +752,7 @@ bool FWaveformTransformationMarkerRenderer::SetMarkerInInteractionRange(const FV
 		if (CuePoint.FrameLength > 0 && IsPositionInLoopArea(InLocalCursorPosition, WidgetGeometry, &CuePoint))
 		{
 			const uint32 StartFrameOffset = TransformationWaveInfo.StartFrameOffset / TransformationWaveInfo.NumChannels;
-			int32 LoopRegionPosition = (InLocalCursorPosition.X / PixelsPerFrame) - StartFrameOffset - CuePoint.FramePosition;
+			const int32 LoopRegionPosition = (InLocalCursorPosition.X / PixelsPerFrame) - CuePoint.FramePosition;
 
 			MarkerInInteractionRange.SetMarkerInInteractionRange(&CuePoint, false, LoopRegionPosition);
 			MarkersArray->SelectedCue = CuePoint.CuePointID;
@@ -793,15 +797,35 @@ void FWaveformTransformationMarkerRenderer::SetActiveLoopRegion(int32 CuePointID
 	MarkersHandle->NotifyPreChange();
 	MarkersArray->Modify();
 	MarkersArray->SelectedCue = CuePointID;
-	EndTransaction();
-	MarkersHandle->NotifyPostChange(EPropertyChangeType::ValueSet);
 
+	check(TransformationWaveInfo.SampleRate > 0);
 	// Revert Loop preview handles
 	if (CuePointID == INDEX_NONE)
 	{
-		check(TransformationWaveInfo.SampleRate > 0);
 		SetIsPreviewingLoopRegion(0.0, -1.0, false);
 	}
+	else
+	{
+		PreviewSelectedLoop();
+	}
+
+	EndTransaction();
+}
+
+void FWaveformTransformationMarkerRenderer::PreviewSelectedLoop()
+{
+	check(TransformationWaveInfo.SampleRate > 0);
+	for (FSoundWaveCuePoint CuePoint : MarkersArray->CuesAndLoops)
+	{
+		if (CuePoint.CuePointID == MarkersArray->SelectedCue)
+		{
+			// Set Loop preview handles
+			const float StartLoopPosInSeconds = static_cast<float>(CuePoint.FramePosition) / TransformationWaveInfo.SampleRate;
+			const float EndLoopPosInSeconds = static_cast<float>(CuePoint.FramePosition + CuePoint.FrameLength) / TransformationWaveInfo.SampleRate;
+
+			SetIsPreviewingLoopRegion(static_cast<double>(StartLoopPosInSeconds), static_cast<double>(EndLoopPosInSeconds), true);
+		}
+	}
 }
 
 void FWaveformTransformationMarkerRenderer::SetIsPreviewingLoopRegion(double InStartTime, double InEndTime, bool bIsPreviewing)
@@ -904,7 +928,9 @@ void FWaveformTransformationMarkerRenderer::SetPropertyValueDependingOnInteracti
 			if (MarkerInInteractionRange.GetMarker())
 			{
 				check(PixelsPerFrame > 0);
-				MarkerInInteractionRange.SetMarkerPosition(LocalCursorXClamped / PixelsPerFrame, false);
+				check(TransformationWaveInfo.NumChannels > 0);
+				const uint32 StartFrameOffset = TransformationWaveInfo.StartFrameOffset / TransformationWaveInfo.NumChannels;
+				MarkerInInteractionRange.SetMarkerPosition(((LocalCursorXClamped + (StartFrameOffset * PixelsPerFrame)) / PixelsPerFrame), false);
 				HighlightedCue = MarkerInInteractionRange.GetMarker()->CuePointID;
 			}
 		}
@@ -914,7 +940,9 @@ void FWaveformTransformationMarkerRenderer::SetPropertyValueDependingOnInteracti
 		if (MarkerInInteractionRange.GetMarker())
 		{
 			check(PixelsPerFrame > 0);
-			MarkerInInteractionRange.SetMarkerPosition(LocalCursorXClamped / PixelsPerFrame, true);
+			check(TransformationWaveInfo.NumChannels > 0);
+			const uint32 StartFrameOffset = TransformationWaveInfo.StartFrameOffset / TransformationWaveInfo.NumChannels;
+			MarkerInInteractionRange.SetMarkerPosition(((LocalCursorXClamped + (StartFrameOffset * PixelsPerFrame)) / PixelsPerFrame), true);
 			HighlightedCue = MarkerInInteractionRange.GetMarker()->CuePointID;
 		}
 	}
diff --git a/Engine/Plugins/Editor/WaveformEditor/Source/WaveformTransformationsWidgets/Private/WaveformTransformationMarkersRenderer.h b/Engine/Plugins/Editor/WaveformEditor/Source/WaveformTransformationsWidgets/Private/WaveformTransformationMarkersRenderer.h
--- a/Engine/Plugins/Editor/WaveformEditor/Source/WaveformTransformationsWidgets/Private/WaveformTransformationMarkersRenderer.h
+++ b/Engine/Plugins/Editor/WaveformEditor/Source/WaveformTransformationsWidgets/Private/WaveformTransformationMarkersRenderer.h
@@ -42,6 +42,7 @@ private:
 	bool SetMarkerInInteractionRange(const FVector2D& InLocalCursorPosition, const FGeometry& WidgetGeometry);
 	bool SelectLoopArea(const FVector2D& InLocalCursorPosition, const FGeometry& WidgetGeometry);
 	void SetActiveLoopRegion(int32 CuePointID);
+	void PreviewSelectedLoop();
 	void SetIsPreviewingLoopRegion(double InStartTime, double InEndTime, bool bIsPreviewing);
 	static FVector2D GetLocalCursorPosition(const FPointerEvent& MouseEvent, const FGeometry& EventGeometry);
 
