Fix build issue:
//UE5/Main - Build Engine Localization - Text conflict from LOCTEXT macro for namespace "WaveformEditor" and key "SetNumericPropertyTransaction". The conflicting sources are "ToggleFadeOut" and "ToggleFadeIn".
#UE-273814

diff --git a/Engine/Plugins/Editor/WaveformEditor/Source/WaveformEditor/Private/WaveformEditor.cpp b/Engine/Plugins/Editor/WaveformEditor/Source/WaveformEditor/Private/WaveformEditor.cpp
--- a/Engine/Plugins/Editor/WaveformEditor/Source/WaveformEditor/Private/WaveformEditor.cpp
+++ b/Engine/Plugins/Editor/WaveformEditor/Source/WaveformEditor/Private/WaveformEditor.cpp
@@ -276,7 +276,7 @@ void FWaveformEditor::ToggleFadeIn()
 {
 	check(SoundWave);
 	check(GEditor != nullptr);
-	GEditor->BeginTransaction(TEXT("PropertyEditor"), LOCTEXT("SetNumericPropertyTransaction", "ToggleFadeIn"), nullptr);
+	GEditor->BeginTransaction(TEXT("PropertyEditor"), LOCTEXT("SetNumericPropertyTransactionToggleFadeIn", "ToggleFadeIn"), nullptr);
 	SoundWave->Modify();
 
 	TObjectPtr<UWaveformTransformationTrimFade> TrimFadeTransformation = GetOrAddTrimFadeTransformation();
@@ -320,7 +320,7 @@ void FWaveformEditor::ToggleFadeOut()
 {
 	check(SoundWave);
 	check(GEditor != nullptr);
-	GEditor->BeginTransaction(TEXT("PropertyEditor"), LOCTEXT("SetNumericPropertyTransaction", "ToggleFadeOut"), nullptr);
+	GEditor->BeginTransaction(TEXT("PropertyEditor"), LOCTEXT("SetNumericPropertyTransactionToggleFadeOut", "ToggleFadeOut"), nullptr);
 	SoundWave->Modify();
 
 	TObjectPtr<UWaveformTransformationTrimFade> TrimFadeTransformation = GetOrAddTrimFadeTransformation();
@@ -363,7 +363,7 @@ void FWaveformEditor::CreateMarker(bool bIsLoopRegion)
 {
 	check(SoundWave);
 	check(GEditor != nullptr);
-	GEditor->BeginTransaction(TEXT("PropertyEditor"), LOCTEXT("SetNumericPropertyTransaction", "CreateMarker"), nullptr);
+	GEditor->BeginTransaction(TEXT("PropertyEditor"), LOCTEXT("SetNumericPropertyTransactionCreateMarker", "CreateMarker"), nullptr);
 	SoundWave->Modify();
 
 	TObjectPtr<UWaveformTransformationMarkers> MarkerTransformation = GetOrAddMarkerTransformation();
@@ -400,7 +400,7 @@ void FWaveformEditor::DeleteMarker()
 {
 	check(SoundWave);
 	check(GEditor != nullptr);
-	GEditor->BeginTransaction(TEXT("PropertyEditor"), LOCTEXT("SetNumericPropertyTransaction", "DeleteMarker"), nullptr);
+	GEditor->BeginTransaction(TEXT("PropertyEditor"), LOCTEXT("SetNumericPropertyTransactionDeleteMarker", "DeleteMarker"), nullptr);
 	SoundWave->Modify();
 
 	TObjectPtr<UWaveformTransformationMarkers> MarkerTransformation = GetOrAddMarkerTransformation();
