[Waveform Editor] - Duplicating Cue Points will copy the Cue Point ID

diff --git a/Engine/Plugins/Editor/WaveformEditor/Source/WaveformTransformations/Private/WaveformTransformationMarkers.cpp b/Engine/Plugins/Editor/WaveformEditor/Source/WaveformTransformations/Private/WaveformTransformationMarkers.cpp
--- a/Engine/Plugins/Editor/WaveformEditor/Source/WaveformTransformations/Private/WaveformTransformationMarkers.cpp
+++ b/Engine/Plugins/Editor/WaveformEditor/Source/WaveformTransformations/Private/WaveformTransformationMarkers.cpp
@@ -174,24 +174,45 @@ void UWaveCueArray::PostEditChangeProperty(FPropertyChangedEvent& PropertyChange
 	Super::PostEditChangeProperty(PropertyChangedEvent);
 
 	const FName PropertyName = PropertyChangedEvent.GetPropertyName();
-	if (PropertyName == GET_MEMBER_NAME_CHECKED(UWaveCueArray, CuesAndLoops) && CuesAndLoops.Num() > 0 &&
-		(PropertyChangedEvent.ChangeType & (EPropertyChangeType::ArrayAdd | EPropertyChangeType::ResetToDefault | 
-			EPropertyChangeType::Unspecified | EPropertyChangeType::Redirected)))
+	if (PropertyName == GET_MEMBER_NAME_CHECKED(UWaveCueArray, CuesAndLoops) && CuesAndLoops.Num() > 0)
 	{
 		int32 LargestCuePointID = 1;
-		for (int32 Index = CuesAndLoops.Num()-1; Index >= 0; Index--)
+
+		for (int32 Index = CuesAndLoops.Num() - 1; Index >= 0; Index--)
 		{
 			if (CuesAndLoops[Index].CuePointID >= LargestCuePointID)
 			{
 				LargestCuePointID = CuesAndLoops[Index].CuePointID + 1;
 			}
 		}
-		//Elements can be inserted at any index in the array
-		for (FSoundWaveCuePoint& CuePoint : CuesAndLoops)
+
+		// When an element is added or reset to default, CuePointID is equal to INDEX_NONE so only those cue points have to be addressed
+		if (PropertyChangedEvent.ChangeType & (EPropertyChangeType::ArrayAdd | EPropertyChangeType::ResetToDefault |
+				EPropertyChangeType::Unspecified | EPropertyChangeType::Redirected))
+		{
+			// Elements can be inserted at any index in the array
+			for (FSoundWaveCuePoint& CuePoint : CuesAndLoops)
+			{
+				if (CuePoint.CuePointID == INDEX_NONE)
+				{
+					CuePoint.CuePointID = LargestCuePointID++;
+				}
+			}
+		}
+		// When an element is duplicated, the CuePointID is also duplicated so the new cue point needs a new unique ID
+		// Only works for adjacent indicies which is sufficient for the way Duplicate works in the details panel
+		else if (PropertyChangedEvent.ChangeType & EPropertyChangeType::Duplicate)
 		{
-			if (CuePoint.CuePointID == INDEX_NONE)
+			int32 PrevCuePointID = INDEX_NONE;
+
+			for (FSoundWaveCuePoint& CuePoint : CuesAndLoops)
 			{
-				CuePoint.CuePointID = LargestCuePointID++;
+				if (CuePoint.CuePointID == PrevCuePointID)
+				{
+					CuePoint.CuePointID = LargestCuePointID++;
+				}
+
+				PrevCuePointID = CuePoint.CuePointID;
 			}
 		}
 	}
