Waveform Editor: Popup clearning fixes

diff --git a/Engine/Plugins/Editor/WaveformEditor/Source/WaveformTransformationsWidgets/Private/WaveformTransformationMarkersRenderer.cpp b/Engine/Plugins/Editor/WaveformEditor/Source/WaveformTransformationsWidgets/Private/WaveformTransformationMarkersRenderer.cpp
--- a/Engine/Plugins/Editor/WaveformEditor/Source/WaveformTransformationsWidgets/Private/WaveformTransformationMarkersRenderer.cpp
+++ b/Engine/Plugins/Editor/WaveformEditor/Source/WaveformTransformationsWidgets/Private/WaveformTransformationMarkersRenderer.cpp
@@ -22,6 +22,12 @@ FWaveformTransformationMarkerRenderer::FWaveformTransformationMarkerRenderer()
 	ApplyWidgetSettings(Settings);
 }
 
+FWaveformTransformationMarkerRenderer::~FWaveformTransformationMarkerRenderer()
+{
+	FSlateApplication::Get().OnFocusChanging().Remove(PopupHandle);
+	FSlateApplication::Get().OnApplicationActivationStateChanged().Remove(ApplicationActivationStateHandle);
+}
+
 int32 FWaveformTransformationMarkerRenderer::OnPaint(const FPaintArgs& Args, const FGeometry& AllottedGeometry, const FSlateRect& MyCullingRect, FSlateWindowElementList& OutDrawElements, int32 LayerId, const FWidgetStyle& InWidgetStyle, bool bParentEnabled) const
 {
 	LayerId = DrawMarkerHandles(AllottedGeometry, OutDrawElements, LayerId);
@@ -500,12 +506,6 @@ FReply FWaveformTransformationMarkerRenderer::OnMouseButtonUp(SWidget& OwnerWidg
 		PreviewSelectedLoop();
 	}
 
-
-	if (TSharedPtr<SWindow> LockedFadeModeMenuWindow = CreateLoopMenuWindow.Pin())
-	{
-		FSlateApplication::Get().RequestDestroyWindow(LockedFadeModeMenuWindow.ToSharedRef());
-	}
-
 	if (MarkersArray == nullptr)
 	{
 		return FReply::Unhandled();
@@ -1014,6 +1014,38 @@ void FWaveformTransformationMarkerRenderer::ShowConvertToLoopAtMarker(const FGeo
 	CreateLoopMenuWindow = MenuWindow;
 
 	FSlateApplication::Get().AddWindow(MenuWindow);
+
+	FSlateApplication::Get().OnFocusChanging().Remove(PopupHandle);
+	FSlateApplication::Get().OnApplicationActivationStateChanged().Remove(ApplicationActivationStateHandle);
+
+	ApplicationActivationStateHandle = FSlateApplication::Get().OnApplicationActivationStateChanged().AddLambda([this](bool isActive)
+		{
+			if (!isActive)
+			{
+				if (TSharedPtr<SWindow> LockedLoopMenuWindow = CreateLoopMenuWindow.Pin())
+				{
+					FSlateApplication::Get().RequestDestroyWindow(LockedLoopMenuWindow.ToSharedRef());
+					FSlateApplication::Get().OnApplicationActivationStateChanged().Remove(ApplicationActivationStateHandle);
+				}
+			}
+		});
+
+	// If focus is lost on the popup, destroy it to prevent popups hanging around
+	PopupHandle = FSlateApplication::Get().OnFocusChanging().AddLambda([this](const FFocusEvent& FocusEvent, const FWeakWidgetPath& WeakWidgetPath
+		, const TSharedPtr<SWidget>& OldWidget, const FWidgetPath& WidgetPath, const TSharedPtr<SWidget>& NewWidget)
+		{
+			if (CreateLoopMenuWindow != nullptr && CreateLoopMenuWindow.IsValid())
+			{
+				if (TSharedPtr<SWindow> LockedLoopMenuWindow = CreateLoopMenuWindow.Pin())
+				{
+					if (OldWidget && !OldWidget->IsHovered() && LockedLoopMenuWindow == OldWidget)
+					{
+						FSlateApplication::Get().RequestDestroyWindow(LockedLoopMenuWindow.ToSharedRef());
+						FSlateApplication::Get().OnFocusChanging().Remove(PopupHandle);
+					}
+				}
+			}
+		});
 }
 
 void FWaveformTransformationMarkerRenderer::FMarkerInInteractionRange::SetMarkerInInteractionRange(FSoundWaveCuePoint* InMarkerPtr, const bool InIsLoopRegionEndMarker, int32 GrabXOffset)
diff --git a/Engine/Plugins/Editor/WaveformEditor/Source/WaveformTransformationsWidgets/Private/WaveformTransformationMarkersRenderer.h b/Engine/Plugins/Editor/WaveformEditor/Source/WaveformTransformationsWidgets/Private/WaveformTransformationMarkersRenderer.h
--- a/Engine/Plugins/Editor/WaveformEditor/Source/WaveformTransformationsWidgets/Private/WaveformTransformationMarkersRenderer.h
+++ b/Engine/Plugins/Editor/WaveformEditor/Source/WaveformTransformationsWidgets/Private/WaveformTransformationMarkersRenderer.h
@@ -12,6 +12,7 @@ class FWaveformTransformationMarkerRenderer : public FWaveformTransformationRend
 {
 public:
 	FWaveformTransformationMarkerRenderer();
+	~FWaveformTransformationMarkerRenderer();
 
 	virtual int32 OnPaint(const FPaintArgs& Args, const FGeometry& AllottedGeometry, const FSlateRect& MyCullingRect, FSlateWindowElementList& OutDrawElements, int32 LayerId, const FWidgetStyle& InWidgetStyle, bool bParentEnabled) const override;
 	virtual void Tick(const FGeometry& AllottedGeometry, const double InCurrentTime, const float InDeltaTime) override;
@@ -104,6 +105,9 @@ private:
 	TSharedPtr<IPropertyHandle> EndLoopTimeHandle;
 	TSharedPtr<IPropertyHandle> IsPreviewingLoopRegionHandle;
 
+	FDelegateHandle PopupHandle;
+	FDelegateHandle ApplicationActivationStateHandle;
+
 	TWeakPtr<SWindow> CreateLoopMenuWindow;
 	const TArray<int32> SizeIncrements = { 1, 10, 50, 100 };
 	const TArray<FKey> IndexKeyCodes = { EKeys::One, EKeys::Two, EKeys::Three, EKeys::Four, EKeys::Five, EKeys::Six, EKeys::Seven, EKeys::Eight, EKeys::Nine };
diff --git a/Engine/Plugins/Editor/WaveformEditor/Source/WaveformTransformationsWidgets/Private/WaveformTransformationTrimFadeRenderer.cpp b/Engine/Plugins/Editor/WaveformEditor/Source/WaveformTransformationsWidgets/Private/WaveformTransformationTrimFadeRenderer.cpp
--- a/Engine/Plugins/Editor/WaveformEditor/Source/WaveformTransformationsWidgets/Private/WaveformTransformationTrimFadeRenderer.cpp
+++ b/Engine/Plugins/Editor/WaveformEditor/Source/WaveformTransformationsWidgets/Private/WaveformTransformationTrimFadeRenderer.cpp
@@ -11,6 +11,12 @@
 #include "Widgets/Input/SButton.h"
 #include "Widgets/SWidget.h"
 
+FWaveformTransformationTrimFadeRenderer::~FWaveformTransformationTrimFadeRenderer()
+{
+	FSlateApplication::Get().OnFocusChanging().Remove(PopupHandle);
+	FSlateApplication::Get().OnApplicationActivationStateChanged().Remove(ApplicationActivationStateHandle);
+}
+
 int32 FWaveformTransformationTrimFadeRenderer::OnPaint(const FPaintArgs& Args, const FGeometry& AllottedGeometry, const FSlateRect& MyCullingRect, FSlateWindowElementList& OutDrawElements, int32 LayerId, const FWidgetStyle& InWidgetStyle, bool bParentEnabled) const
 {
 	LayerId = DrawTrimHandles(AllottedGeometry, OutDrawElements, LayerId);
@@ -395,6 +401,38 @@ void FWaveformTransformationTrimFadeRenderer::ShowSelectFadeModeMenuAtCursor(con
 	FadeModeMenuWindow = MenuWindow;
 
 	FSlateApplication::Get().AddWindow(MenuWindow);
+
+	FSlateApplication::Get().OnFocusChanging().Remove(PopupHandle);
+	FSlateApplication::Get().OnApplicationActivationStateChanged().Remove(ApplicationActivationStateHandle);
+
+	ApplicationActivationStateHandle = FSlateApplication::Get().OnApplicationActivationStateChanged().AddLambda([this](bool isActive)
+		{
+			if (!isActive)
+			{
+				if (TSharedPtr<SWindow> LockedFadeModeMenuWindow = FadeModeMenuWindow.Pin())
+				{
+					FSlateApplication::Get().RequestDestroyWindow(LockedFadeModeMenuWindow.ToSharedRef());
+					FSlateApplication::Get().OnApplicationActivationStateChanged().Remove(ApplicationActivationStateHandle);
+				}
+			}
+		});
+
+	// If focus is lost on the popup, destroy it to prevent popups hanging around
+	PopupHandle = FSlateApplication::Get().OnFocusChanging().AddLambda([this](const FFocusEvent& FocusEvent, const FWeakWidgetPath& WeakWidgetPath
+		, const TSharedPtr<SWidget>& OldWidget, const FWidgetPath& WidgetPath, const TSharedPtr<SWidget>& NewWidget)
+		{
+			if (FadeModeMenuWindow != nullptr && FadeModeMenuWindow.IsValid())
+			{
+				if (TSharedPtr<SWindow> LockedFadeModeMenuWindow = FadeModeMenuWindow.Pin())
+				{
+					if (OldWidget && !OldWidget->IsHovered() && LockedFadeModeMenuWindow == OldWidget)
+					{
+						FSlateApplication::Get().RequestDestroyWindow(LockedFadeModeMenuWindow.ToSharedRef());
+						FSlateApplication::Get().OnFocusChanging().Remove(PopupHandle);
+					}
+				}
+			}
+		});
 }
 
 FWaveformTransformationTrimFadeRenderer::ETrimFadeInteractionType FWaveformTransformationTrimFadeRenderer::GetInteractionTypeFromCursorPosition(const FVector2D& InLocalCursorPosition, const FKey MouseButton, const FGeometry& WidgetGeometry) const
diff --git a/Engine/Plugins/Editor/WaveformEditor/Source/WaveformTransformationsWidgets/Private/WaveformTransformationTrimFadeRenderer.h b/Engine/Plugins/Editor/WaveformEditor/Source/WaveformTransformationsWidgets/Private/WaveformTransformationTrimFadeRenderer.h
--- a/Engine/Plugins/Editor/WaveformEditor/Source/WaveformTransformationsWidgets/Private/WaveformTransformationTrimFadeRenderer.h
+++ b/Engine/Plugins/Editor/WaveformEditor/Source/WaveformTransformationsWidgets/Private/WaveformTransformationTrimFadeRenderer.h
@@ -8,6 +8,7 @@ class FWaveformTransformationTrimFadeRenderer : public FWaveformTransformationRe
 {
 public:
 	FWaveformTransformationTrimFadeRenderer() = default;
+	~FWaveformTransformationTrimFadeRenderer();
 
 	virtual int32 OnPaint(const FPaintArgs& Args, const FGeometry& AllottedGeometry, const FSlateRect& MyCullingRect, FSlateWindowElementList& OutDrawElements, int32 LayerId, const FWidgetStyle& InWidgetStyle, bool bParentEnabled) const override;
 	virtual void Tick(const FGeometry& AllottedGeometry, const double InCurrentTime, const float InDeltaTime) override;
@@ -73,4 +74,7 @@ private:
 	TSharedPtr<IPropertyHandle> EndFadeTimeHandle;
 	TSharedPtr<IPropertyHandle> EndFadeCurveHandle;
 	TSharedPtr<IPropertyHandle> EndSCurveSharpnessHandle;
+
+	FDelegateHandle PopupHandle;
+	FDelegateHandle ApplicationActivationStateHandle;
 };
