[Waveform Editor] - Crash when using negative Frame Length on a Loop Region
[Waveform Editor] - Sample Data warning when adjusting a Cue Points position
[Waveform Editor] - Crash when moving a Cue Point in a SoundWave with more than one channel. Fixes the crash in the waveform editor but does not fix the underlying ensure issue.
[Waveform Editor] - Crash when zooming in on a Waveform
Markers are now visible when outside a trim scope.
Changes to WaveformEditorSequenceProvider so that the entire waveform view is available to all transformation renderers

diff --git a/Engine/Plugins/Editor/WaveformEditor/Source/WaveformEditor/Private/WaveformEditor.cpp b/Engine/Plugins/Editor/WaveformEditor/Source/WaveformEditor/Private/WaveformEditor.cpp
--- a/Engine/Plugins/Editor/WaveformEditor/Source/WaveformEditor/Private/WaveformEditor.cpp
+++ b/Engine/Plugins/Editor/WaveformEditor/Source/WaveformEditor/Private/WaveformEditor.cpp
@@ -383,12 +383,12 @@ void FWaveformEditor::CreateMarker(bool bIsLoopRegion)
 
 	FSoundWaveCuePoint NewCuePoint;
 	NewCuePoint.CuePointID = HighestCueID + 1;
-	NewCuePoint.FramePosition = TransportCoordinator->GetFocusPoint() * SoundWave->TotalSamples;
+	NewCuePoint.FramePosition = static_cast<int64>(TransportCoordinator->GetFocusPoint() * SoundWave->TotalSamples);
 
 	if (bIsLoopRegion)
 	{
 		NewCuePoint.SetLoopRegion(true);
-		NewCuePoint.FrameLength = SoundWave->TotalSamples * 0.1; // Default loop region to 10% of samples for easier tuning
+		NewCuePoint.FrameLength = static_cast<int64>(SoundWave->TotalSamples * 0.1f); // Default loop region to 10% of samples for easier tuning
 	}
 
 	MarkerTransformation->Markers->CuesAndLoops.Add(NewCuePoint);
@@ -439,8 +439,8 @@ void FWaveformEditor::SkipToNextMarker()
 	{
 		check(TransportCoordinator != nullptr);
 		check(SoundWave->TotalSamples > 0);
-		const float PlaybackFramePosition = TransportCoordinator->GetFocusPoint() * SoundWave->TotalSamples;
-		float NewPlaybackFramePosition = FLT_MAX;
+		const int64 PlaybackFramePosition = static_cast<int64>(TransportCoordinator->GetFocusPoint() * SoundWave->TotalSamples);
+		int64 NewPlaybackFramePosition = MAX_int64;
 
 		// Find the nearest marker beyond the current playhead.
 		check(MarkerTransformation->Markers != nullptr);
@@ -455,7 +455,7 @@ void FWaveformEditor::SkipToNextMarker()
 			}
 		}
 
-		if (NewPlaybackFramePosition != FLT_MAX)
+		if (NewPlaybackFramePosition != MAX_int64)
 		{
 			float ActiveDuration = TransformationChainConfig.EndTime - TransformationChainConfig.StartTime;
 			if (ActiveDuration <= 0.0f)
diff --git a/Engine/Plugins/Editor/WaveformEditor/Source/WaveformEditorWidgets/Private/WaveformEditorSequenceDataProvider.cpp b/Engine/Plugins/Editor/WaveformEditor/Source/WaveformEditorWidgets/Private/WaveformEditorSequenceDataProvider.cpp
--- a/Engine/Plugins/Editor/WaveformEditor/Source/WaveformEditorWidgets/Private/WaveformEditorSequenceDataProvider.cpp
+++ b/Engine/Plugins/Editor/WaveformEditor/Source/WaveformEditorWidgets/Private/WaveformEditorSequenceDataProvider.cpp
@@ -98,10 +98,10 @@ void FWaveformEditorSequenceDataProvider::GenerateSequenceDataInternal()
 {
 	check(SoundWaveToRender);
 	
-	uint32 NumWaveformSamples = ImportedRawPCMData.Num() * sizeof(uint8) / sizeof(int16);
+	int64 NumWaveformSamples = ImportedRawPCMData.Num() * sizeof(uint8) / sizeof(int16);
 	NumOriginalWaveformFrames = NumWaveformSamples / ImportedNumChannels;
-	uint32 FirstEditedSample = 0;
-	uint32 LastEditedSample = NumWaveformSamples;
+	int64 FirstEditedSample = 0;
+	int64 LastEditedSample = NumWaveformSamples;
 
 	if (TransformationsToRender.Num() > 0)
 	{
@@ -114,6 +114,7 @@ void FWaveformEditorSequenceDataProvider::GenerateSequenceDataInternal()
 
 		Audio::ArrayPcm16ToFloat(MakeArrayView((int16*)ImportedRawPCMData.GetData(), NumWaveformSamples), TransformationsBuffer);
 		OutputBuffer = TransformationsBuffer;
+		check(OutputBuffer.Num() > 0);
 
 		TransformationInfo.Audio = &TransformationsBuffer;
 		TransformationInfo.NumChannels = ImportedNumChannels;
@@ -128,19 +129,21 @@ void FWaveformEditorSequenceDataProvider::GenerateSequenceDataInternal()
 
 			if (RenderLayers[i].Key)
 			{
-				FWaveformTransformationRenderInfo RenderLayerInfo{ TransformationInfo.SampleRate, TransformationInfo.NumChannels, FirstEditedSample, LastEditedSample - FirstEditedSample };
+				FWaveformTransformationRenderInfo RenderLayerInfo{ TransformationInfo.SampleRate, TransformationInfo.NumChannels, FirstEditedSample, NumWaveformSamples, LastEditedSample - FirstEditedSample };
 				RenderLayers[i].Key->SetTransformationWaveInfo(MoveTemp(RenderLayerInfo));
 
-				float LeftWidgetConstraint = FirstEditedSample / (float) NumWaveformSamples;
-				float RightWidgetConstraint = LastEditedSample / (float) NumWaveformSamples;
-
-				RenderLayers[i].Value = FTransformationLayerConstraints(LeftWidgetConstraint, RightWidgetConstraint);
+				// For rendering, we want to be able to draw over all the samples, not just the ones available after processing
+				RenderLayers[i].Value = FTransformationLayerConstraints(0, 1);
 
 			}
 
 			if (bChainChangesFileLength)
 			{
 				FirstEditedSample += TransformationInfo.StartFrameOffset;
+				if (FirstEditedSample >= OutputBuffer.Num())
+				{
+					FirstEditedSample = OutputBuffer.Num() - 1;
+				}
 
 				//Prevent a lower priority transformation from setting the LastEditedSample past the end of higher priority transformations
 				if (Transformations[i]->FileChangeLengthPriority() == Audio::ETransformationPriority::High
@@ -150,7 +153,11 @@ void FWaveformEditorSequenceDataProvider::GenerateSequenceDataInternal()
 					LastEditedSample = TransformationInfo.NumEditedSamples <= 0 ? LastEditedSample : FirstEditedSample + TransformationInfo.NumEditedSamples;
 				}
 			
-				check(LastEditedSample > FirstEditedSample);
+				if(LastEditedSample <= FirstEditedSample)
+				{
+					LastEditedSample = FirstEditedSample + 1;
+				}
+
 				const int32 EditedBufferSize = LastEditedSample - FirstEditedSample;
 				FMemory::Memcpy(&OutputBuffer.GetData()[FirstEditedSample], TransformationsBuffer.GetData(), EditedBufferSize * sizeof(float));
 			}
diff --git a/Engine/Plugins/Editor/WaveformEditor/Source/WaveformTransformations/Private/WaveformTransformationMarkers.cpp b/Engine/Plugins/Editor/WaveformEditor/Source/WaveformTransformations/Private/WaveformTransformationMarkers.cpp
--- a/Engine/Plugins/Editor/WaveformEditor/Source/WaveformTransformations/Private/WaveformTransformationMarkers.cpp
+++ b/Engine/Plugins/Editor/WaveformEditor/Source/WaveformTransformations/Private/WaveformTransformationMarkers.cpp
@@ -23,9 +23,9 @@ void FWaveTransformationMarkers::ProcessAudio(Audio::FWaveformTransformationWave
 	check(StartLoopTime >= 0);
 	check(InOutWaveInfo.NumChannels > 0);
 
-	const int32 StartSampleOffset = InOutWaveInfo.StartFrameOffset;
+	const int64 StartSampleOffset = InOutWaveInfo.StartFrameOffset;
 
-	int32 StartSample = FMath::RoundToInt32(StartLoopTime * InOutWaveInfo.SampleRate) * InOutWaveInfo.NumChannels;
+	int64 StartSample = FMath::RoundToInt32(StartLoopTime * InOutWaveInfo.SampleRate) * InOutWaveInfo.NumChannels;
 	StartSample = FMath::Max(0, StartSample - StartSampleOffset);
 
 	if (StartSample > InputAudio.Num())
@@ -33,12 +33,12 @@ void FWaveTransformationMarkers::ProcessAudio(Audio::FWaveformTransformationWave
 		return;
 	}
 	
-	int32 EndSample = InputAudio.Num() - 1;
+	int64 EndSample = InputAudio.Num() - 1;
 
 	if (EndLoopTime > 0.f)
 	{
-		const int32 EndFrame = FMath::RoundToInt32(EndLoopTime * InOutWaveInfo.SampleRate) * InOutWaveInfo.NumChannels;
-		EndSample = FMath::Max(StartSample, EndFrame - 1 - StartSampleOffset);
+		const int64 EndFrame = FMath::RoundToInt32(EndLoopTime * InOutWaveInfo.SampleRate) * InOutWaveInfo.NumChannels;
+		EndSample = FMath::Max(StartSample + 1, EndFrame - 1 - StartSampleOffset);
 
 		// EndLoopTime can be beyond the length of the file if there is a Trim
 		if (EndSample > InputAudio.Num() - 1)
@@ -52,6 +52,13 @@ void FWaveTransformationMarkers::ProcessAudio(Audio::FWaveformTransformationWave
 
 	const int32 FinalSize = EndSample - StartSample + 1;
 
+	if (FinalSize <= 2)
+	{
+		UE_LOG(LogWaveformTransformationMarkers, Warning, TEXT("Previewing loop of sample size 1!"));
+		
+		return;
+	}
+
 	InOutWaveInfo.StartFrameOffset = StartSample - (StartSample % InOutWaveInfo.NumChannels);
 	InOutWaveInfo.NumEditedSamples = FinalSize;
 
@@ -173,6 +180,14 @@ void UWaveCueArray::PostEditChangeProperty(FPropertyChangedEvent& PropertyChange
 {
 	Super::PostEditChangeProperty(PropertyChangedEvent);
 
+	for (FSoundWaveCuePoint& CuePoint : CuesAndLoops)
+	{
+		if (CuePoint.IsLoopRegion() && CuePoint.FrameLength < MinLoopSize)
+		{
+			CuePoint.FrameLength = MinLoopSize;
+		}
+	}
+
 	const FName PropertyName = PropertyChangedEvent.GetPropertyName();
 	if (PropertyName == GET_MEMBER_NAME_CHECKED(UWaveCueArray, CuesAndLoops) && CuesAndLoops.Num() > 0)
 	{
diff --git a/Engine/Plugins/Editor/WaveformEditor/Source/WaveformTransformations/Public/WaveformTransformationMarkers.h b/Engine/Plugins/Editor/WaveformEditor/Source/WaveformTransformations/Public/WaveformTransformationMarkers.h
--- a/Engine/Plugins/Editor/WaveformEditor/Source/WaveformTransformations/Public/WaveformTransformationMarkers.h
+++ b/Engine/Plugins/Editor/WaveformEditor/Source/WaveformTransformations/Public/WaveformTransformationMarkers.h
@@ -48,6 +48,9 @@ public:
 	DECLARE_DELEGATE_OneParam(CycleMarkerLoopRegionDelegate, ELoopModificationControls);
 	CycleMarkerLoopRegionDelegate CycleMarkerLoop;
 
+	// To minimize complexity while supporting all common editing cases, loops have a min length of 10 frames
+	static constexpr int64 MinLoopSize = 10;
+
 private:
 	UPROPERTY()
 	bool bIsInitialized = false;
diff --git a/Engine/Plugins/Editor/WaveformEditor/Source/WaveformTransformationsWidgets/Private/WaveformTransformationDurationRenderer.cpp b/Engine/Plugins/Editor/WaveformEditor/Source/WaveformTransformationsWidgets/Private/WaveformTransformationDurationRenderer.cpp
--- a/Engine/Plugins/Editor/WaveformEditor/Source/WaveformTransformationsWidgets/Private/WaveformTransformationDurationRenderer.cpp
+++ b/Engine/Plugins/Editor/WaveformEditor/Source/WaveformTransformationsWidgets/Private/WaveformTransformationDurationRenderer.cpp
@@ -9,9 +9,18 @@ FWaveformTransformationDurationRenderer::FWaveformTransformationDurationRenderer
 
 int32 FWaveformTransformationDurationRenderer::OnPaint(const FPaintArgs& Args, const FGeometry& AllottedGeometry, const FSlateRect& MyCullingRect, FSlateWindowElementList& OutDrawElements, int32 LayerId, const FWidgetStyle& InWidgetStyle, bool bParentEnabled) const
 {
-	const float StartTimeRatio = (float)TransformationWaveInfo.StartFrameOffset / OriginalWaveformNumFrames / TransformationWaveInfo.NumChannels;
-	const uint32 EndSample = TransformationWaveInfo.StartFrameOffset + TransformationWaveInfo.NumAvilableSamples;
-	const float EndTimeRatio = (float)EndSample / OriginalWaveformNumFrames / TransformationWaveInfo.NumChannels;
+	// If there is a null transformation in Transformations list, the TransformationWaveInfo has not been set up
+	if (TransformationWaveInfo.NumChannels <= 0)
+	{
+		return LayerId;
+	}
+
+	float OriginalWaveformNumFramesDen = static_cast<float>(OriginalWaveformNumFrames * TransformationWaveInfo.NumChannels);
+	check(OriginalWaveformNumFramesDen > 0.f);
+
+	const float StartTimeRatio = static_cast<float>(TransformationWaveInfo.StartFrameOffset) / OriginalWaveformNumFramesDen;
+	const int64 EndSample = TransformationWaveInfo.StartFrameOffset + static_cast<int64>(TransformationWaveInfo.TotalNumSamples);
+	const float EndTimeRatio = static_cast<float>(EndSample) / OriginalWaveformNumFramesDen;
 
 	const bool bRenderLeftBox = StartTimeRatio >= 0.f;
 	const bool bRenderRightBox = EndTimeRatio <= 1.f;
diff --git a/Engine/Plugins/Editor/WaveformEditor/Source/WaveformTransformationsWidgets/Private/WaveformTransformationMarkersRenderer.cpp b/Engine/Plugins/Editor/WaveformEditor/Source/WaveformTransformationsWidgets/Private/WaveformTransformationMarkersRenderer.cpp
--- a/Engine/Plugins/Editor/WaveformEditor/Source/WaveformTransformationsWidgets/Private/WaveformTransformationMarkersRenderer.cpp
+++ b/Engine/Plugins/Editor/WaveformEditor/Source/WaveformTransformationsWidgets/Private/WaveformTransformationMarkersRenderer.cpp
@@ -41,8 +41,6 @@ int32 FWaveformTransformationMarkerRenderer::DrawMarkerHandles(const FGeometry&
 	LinePoints.SetNumUninitialized(2);
 
 	check(TransformationWaveInfo.NumChannels > 0);
-	check(TransformationWaveInfo.StartFrameOffset % TransformationWaveInfo.NumChannels == 0);
-	const uint32 StartFrameOffset = TransformationWaveInfo.StartFrameOffset / TransformationWaveInfo.NumChannels;
 	const float MarkerHeight = AllottedGeometry.Size.Y * InteractionRatioYDelta;
 	constexpr float MarkerWidth = InteractionPixelXDelta * 0.8;
 	FSlateRoundedBoxBrush RoundedBoxBrush(FLinearColor::Red, 2);
@@ -57,8 +55,7 @@ int32 FWaveformTransformationMarkerRenderer::DrawMarkerHandles(const FGeometry&
 	}
 	for (const FSoundWaveCuePoint& LoopRegion : LoopPoints)
 	{
-		//Positions that become negative after StartFrameOffset subtraction are not rendered
-		const float LoopRegionPosition = (LoopRegion.FramePosition - StartFrameOffset) * PixelsPerFrame;
+		const float LoopRegionPosition = static_cast<float>(LoopRegion.FramePosition) * PixelsPerFrame;
 		LinePoints[0] = FVector2D(LoopRegionPosition, 0.f);
 		LinePoints[1] = FVector2D(LoopRegionPosition, AllottedGeometry.Size.Y);
 
@@ -77,7 +74,7 @@ int32 FWaveformTransformationMarkerRenderer::DrawMarkerHandles(const FGeometry&
 
 		const float MarkerCenter = LoopRegionPosition;
 		const float	HandleStart = MarkerCenter - MarkerWidth / 2;
-		const float LoopRegionPixelLength = LoopRegion.FrameLength * PixelsPerFrame;
+		const float LoopRegionPixelLength = static_cast<float>(LoopRegion.FrameLength) * PixelsPerFrame;
 		FSlateBrush Brush;
 		FPaintGeometry HandleGeometry = AllottedGeometry.ToPaintGeometry(FVector2f(MarkerWidth, MarkerHeight), FSlateLayoutTransform(FVector2f(HandleStart, 0)));
 		FPaintGeometry LoopBoxGeometry = AllottedGeometry.ToPaintGeometry(FVector2f(LoopRegionPixelLength, AllottedGeometry.GetLocalSize().Y), FSlateLayoutTransform(FVector2f(MarkerCenter, 0)));
@@ -118,7 +115,7 @@ int32 FWaveformTransformationMarkerRenderer::DrawMarkerHandles(const FGeometry&
 		);
 
 		// If the loop region has no width or too small to draw handles, use the marker to move it
-		if (LoopRegion.FrameLength < LoopHandlePixelWidth * 2)
+		if (static_cast<float>(LoopRegion.FrameLength) < LoopHandlePixelWidth * 2.f)
 		{
 			FSlateDrawElement::MakeBox(
 				OutDrawElements,
@@ -142,7 +139,7 @@ int32 FWaveformTransformationMarkerRenderer::DrawMarkerHandles(const FGeometry&
 		);
 
 		// If we have a loop region highlighted by the mouse, draw the handles
-		const bool DrawRegionHandles = (HighlightedCue != INDEX_NONE && HighlightedCue == LoopRegion.CuePointID && LoopRegion.FrameLength > LoopHandlePixelWidth);
+		const bool DrawRegionHandles = (HighlightedCue != INDEX_NONE && HighlightedCue == LoopRegion.CuePointID && static_cast<float>(LoopRegion.FrameLength) > LoopHandlePixelWidth);
 
 		if (DrawRegionHandles)
 		{
@@ -177,8 +174,7 @@ int32 FWaveformTransformationMarkerRenderer::DrawMarkerHandles(const FGeometry&
 
 	for (const FSoundWaveCuePoint& Marker : CuePoints)
 	{
-		//Positions that become negative after StartFrameOffset subtraction are not rendered
-		const float MarkerPosition = (Marker.FramePosition - StartFrameOffset) * PixelsPerFrame;
+		const float MarkerPosition = static_cast<float>(Marker.FramePosition) * static_cast<float>(PixelsPerFrame);
 		LinePoints[0] = FVector2D(MarkerPosition, 0.f);
 		LinePoints[1] = FVector2D(MarkerPosition, AllottedGeometry.Size.Y);
 
@@ -245,11 +241,13 @@ void FWaveformTransformationMarkerRenderer::ModifyMarkerLoopRegion(ELoopModifica
 
 	check(MarkersArray);
 	check(TransformationWaveInfo.NumChannels > 0);
-	check(TransformationWaveInfo.NumAvilableSamples > 0);
+	check(TransformationWaveInfo.TotalNumSamples > 0);
 	check(SizeIncrements.Num() > 0);
 
-	int32 FramesToShift = (SizeIncrements[SelectedIncrement] / 1000.0f) * TransformationWaveInfo.SampleRate;
-	int32 MaxFrames = TransformationWaveInfo.NumAvilableSamples;
+	int64 FramesToShift = static_cast<int64>((SizeIncrements[SelectedIncrement] / 1000.0f) * TransformationWaveInfo.SampleRate);
+	int64 MaxFrames = static_cast<int64>(TransformationWaveInfo.TotalNumSamples) / static_cast<int64>(TransformationWaveInfo.NumChannels);
+	check(MaxFrames > 0);
+	
 	BeginTransaction(TEXT("PropertyEditor"), LOCTEXT("SetNumericPropertyTransaction", "Modifying Marker"), nullptr);
 
 	MarkersArray->Modify();
@@ -299,8 +297,11 @@ void FWaveformTransformationMarkerRenderer::ModifyMarkerLoopRegion(ELoopModifica
 		break;
 	}
 
-	LoopRegion->FramePosition = FMath::Clamp(LoopRegion->FramePosition, 0, LoopRegion->FramePosition);
-	LoopRegion->FrameLength = FMath::Clamp(LoopRegion->FrameLength, 0, MaxFrames - LoopRegion->FramePosition);
+	// Prevent the frame position from exceeding the wave length
+	LoopRegion->FramePosition = FMath::Clamp(LoopRegion->FramePosition, 0, MaxFrames - 1);
+	// FrameLength can exceed MaxFrames because the loop will be cut by ProcessAudio
+	// but length must be > 0 or it is no longer a loop
+	LoopRegion->FrameLength = FMath::Max(LoopRegion->FrameLength, 1);
 
 	PreviewSelectedLoop();
 
@@ -422,10 +423,10 @@ void FWaveformTransformationMarkerRenderer::SelectLoopRegionByKeyboard(const FKe
 void FWaveformTransformationMarkerRenderer::Tick(const FGeometry& AllottedGeometry, const double InCurrentTime, const float InDeltaTime)
 { 
 	check(TransformationWaveInfo.NumChannels > 0);
-	const float NumFrames = TransformationWaveInfo.NumAvilableSamples / TransformationWaveInfo.NumChannels;
+	const float NumFrames = static_cast<float>(TransformationWaveInfo.TotalNumSamples) / static_cast<float>(TransformationWaveInfo.NumChannels);
 	check(NumFrames > 0);
-	PixelsPerFrame = FMath::Max(AllottedGeometry.GetLocalSize().X / NumFrames, 0.0);
-	ensure(PixelsPerFrame >= 0);
+	PixelsPerFrame = FMath::Max(AllottedGeometry.GetLocalSize().X / NumFrames, UE_SMALL_NUMBER);
+	ensure(PixelsPerFrame > 0);
 	MarkerInInteractionRange.SetMinLoopSize((LoopHandlePixelWidth / PixelsPerFrame) * 3);
 
 	FVector2D MouseAbsolutePosition = UWidgetLayoutLibrary::GetMousePositionOnPlatform();
@@ -433,7 +434,7 @@ void FWaveformTransformationMarkerRenderer::Tick(const FGeometry& AllottedGeomet
 
 	check(MarkersArray != nullptr);
 
-	TMap<uint32, FSoundWaveCuePoint> CuePointMap;
+	TMap<int32, FSoundWaveCuePoint> CuePointMap;
 	CuePointMap.Reserve(MarkersArray->CuesAndLoops.Num());
 
 	bool IsMarkerInInteractionRangeValid = false;
@@ -455,6 +456,22 @@ void FWaveformTransformationMarkerRenderer::Tick(const FGeometry& AllottedGeomet
 		{
 			IsMarkerInInteractionRangeValid = true;
 		}
+
+		check(TransformationWaveInfo.NumChannels > 0);
+		int64 StartFrameOffset = TransformationWaveInfo.StartFrameOffset / TransformationWaveInfo.NumChannels;
+		int64 NumAvailableSamples = static_cast<int64>(TransformationWaveInfo.NumSamplesAvailable) / TransformationWaveInfo.NumChannels;
+
+		// If loop region is no longer within the available frames, revert Loop preview handles 
+		if (Marker.CuePointID == MarkersArray->SelectedCue && Marker.IsLoopRegion() &&
+			(Marker.FramePosition + Marker.FrameLength < StartFrameOffset || Marker.FramePosition > StartFrameOffset + NumAvailableSamples))
+		{
+			MarkersHandle->NotifyPreChange();
+			MarkersArray->Modify();
+			MarkersArray->SelectedCue = INDEX_NONE;
+			
+			SetIsPreviewingLoopRegion(0.0, -1.0, false);
+			EndTransaction();
+		}
 	}
 
 	//Invalidate Marker pointer if it has been removed from the array
@@ -465,7 +482,7 @@ void FWaveformTransformationMarkerRenderer::Tick(const FGeometry& AllottedGeomet
 
 	CuePoints.Empty();
 	LoopPoints.Empty();
-	for (TPair<uint32, FSoundWaveCuePoint>& Pair : CuePointMap)
+	for (TPair<int32, FSoundWaveCuePoint>& Pair : CuePointMap)
 	{
 		if (Pair.Value.IsLoopRegion())
 		{
@@ -650,9 +667,9 @@ void FWaveformTransformationMarkerRenderer::SetLoopColors(const TArray<FLinearCo
 	}
 }
 
-bool FWaveformTransformationMarkerRenderer::IsPositionInInteractionRange(const int32 InFramePosition, const FVector2D& InLocalCursorPosition, const FGeometry& WidgetGeometry) const
+bool FWaveformTransformationMarkerRenderer::IsPositionInInteractionRange(const int64 InFramePosition, const FVector2D& InLocalCursorPosition, const FGeometry& WidgetGeometry) const
 {
-	TRange InteractionRange(TRange<float>::Inclusive(InFramePosition * PixelsPerFrame - InteractionPixelXDelta, InFramePosition * PixelsPerFrame + InteractionPixelXDelta));
+	TRange InteractionRange(TRange<float>::Inclusive(static_cast<float>(InFramePosition) * PixelsPerFrame - InteractionPixelXDelta, static_cast<float>(InFramePosition) * PixelsPerFrame + InteractionPixelXDelta));
 	if (InteractionRange.Contains(InLocalCursorPosition.X) &&
 		InLocalCursorPosition.Y < WidgetGeometry.GetLocalSize().Y * InteractionRatioYDelta)
 	{
@@ -668,9 +685,8 @@ bool FWaveformTransformationMarkerRenderer::IsPositionInLoopArea(const FVector2D
 	check(CueMarker != nullptr);
 	ensure(LoopHandlePixelWidth > 0);
 
-	const uint32 StartFrameOffset = TransformationWaveInfo.StartFrameOffset / TransformationWaveInfo.NumChannels;
-	float LoopRegionPosition = (CueMarker->FramePosition - StartFrameOffset) * PixelsPerFrame;
-	float LoopRegionPixelLength = CueMarker->FrameLength * PixelsPerFrame;
+	float LoopRegionPosition = static_cast<float>(CueMarker->FramePosition) * static_cast<float>(PixelsPerFrame);
+	float LoopRegionPixelLength = static_cast<float>(CueMarker->FrameLength) * static_cast<float>(PixelsPerFrame);
 
 	// If we want the handles to be included for the area check, as having the mouse in the handle area has different behaviour to the total loop area
 	if (!IncludeHandleArea)
@@ -695,14 +711,13 @@ bool FWaveformTransformationMarkerRenderer::IsPositionInLoopHandleArea(const FVe
 	check(TransformationWaveInfo.NumChannels > 0);
 	ensure(LoopHandlePixelWidth > 0);
 
-	const uint32 StartFrameOffset = TransformationWaveInfo.StartFrameOffset / TransformationWaveInfo.NumChannels;
 	float LoopRegionPosition = IsLeftHandle ?
-		(CueMarker.FramePosition - StartFrameOffset) * PixelsPerFrame :
-		((CueMarker.FramePosition - StartFrameOffset + CueMarker.FrameLength) * PixelsPerFrame) - LoopHandlePixelWidth;
+		static_cast<float>(CueMarker.FramePosition) * static_cast<float>(PixelsPerFrame) :
+		(static_cast<float>(CueMarker.FramePosition + CueMarker.FrameLength) * static_cast<float>(PixelsPerFrame)) - LoopHandlePixelWidth;
 
 	// Check that only the handle within the loop area is highlighted
-	if (InLocalCursorPosition.X > LoopRegionPosition && InLocalCursorPosition.X < LoopRegionPosition + LoopHandlePixelWidth &&
-		InLocalCursorPosition.Y < WidgetGeometry.GetLocalSize().Y && InLocalCursorPosition.Y > 0)
+	if (static_cast<float>(InLocalCursorPosition.X) > LoopRegionPosition && static_cast<float>(InLocalCursorPosition.X) < LoopRegionPosition + LoopHandlePixelWidth &&
+		static_cast<float>(InLocalCursorPosition.Y) < WidgetGeometry.GetLocalSize().Y && InLocalCursorPosition.Y > 0.0)
 	{
 		return true;
 	}
@@ -751,8 +766,7 @@ bool FWaveformTransformationMarkerRenderer::SetMarkerInInteractionRange(const FV
 	{
 		if (CuePoint.FrameLength > 0 && IsPositionInLoopArea(InLocalCursorPosition, WidgetGeometry, &CuePoint))
 		{
-			const uint32 StartFrameOffset = TransformationWaveInfo.StartFrameOffset / TransformationWaveInfo.NumChannels;
-			const int32 LoopRegionPosition = (InLocalCursorPosition.X / PixelsPerFrame) - CuePoint.FramePosition;
+			const int64 LoopRegionPosition = static_cast<int64>(InLocalCursorPosition.X / PixelsPerFrame - static_cast<double>(CuePoint.FramePosition));
 
 			MarkerInInteractionRange.SetMarkerInInteractionRange(&CuePoint, false, LoopRegionPosition);
 			MarkersArray->SelectedCue = CuePoint.CuePointID;
@@ -772,8 +786,6 @@ bool FWaveformTransformationMarkerRenderer::SelectLoopArea(const FVector2D& InLo
 	check(TransformationWaveInfo.NumChannels > 0);
 	ensure(LoopHandlePixelWidth > 0);
 
-	const uint32 StartFrameOffset = TransformationWaveInfo.StartFrameOffset / TransformationWaveInfo.NumChannels;
-
 	for (FSoundWaveCuePoint& LoopRegion : LoopPoints)
 	{
 		// Check mouse position bounds
@@ -815,13 +827,24 @@ void FWaveformTransformationMarkerRenderer::SetActiveLoopRegion(int32 CuePointID
 void FWaveformTransformationMarkerRenderer::PreviewSelectedLoop()
 {
 	check(TransformationWaveInfo.SampleRate > 0);
+	check(TransformationWaveInfo.NumChannels > 0);
+
+	const int64 StartFrameOffset = TransformationWaveInfo.StartFrameOffset / TransformationWaveInfo.NumChannels;
+	const int64 EndFramePosition = StartFrameOffset + TransformationWaveInfo.NumSamplesAvailable / TransformationWaveInfo.NumChannels;
+
 	for (FSoundWaveCuePoint CuePoint : MarkersArray->CuesAndLoops)
 	{
-		if (CuePoint.CuePointID == MarkersArray->SelectedCue)
+		const int64 LoopEndFramePosition = CuePoint.FramePosition + CuePoint.FrameLength;
+		
+		// Only preview a loop if part of it is within the available frames (accounting for TrimFades and other transformations)
+		if (CuePoint.CuePointID == MarkersArray->SelectedCue && CuePoint.IsLoopRegion() &&
+			LoopEndFramePosition >= StartFrameOffset && CuePoint.FramePosition <= EndFramePosition)
 		{
+			check(CuePoint.FrameLength > 0);
+
 			// Set Loop preview handles
 			const float StartLoopPosInSeconds = static_cast<float>(CuePoint.FramePosition) / TransformationWaveInfo.SampleRate;
-			const float EndLoopPosInSeconds = static_cast<float>(CuePoint.FramePosition + CuePoint.FrameLength) / TransformationWaveInfo.SampleRate;
+			const float EndLoopPosInSeconds = static_cast<float>(LoopEndFramePosition) / TransformationWaveInfo.SampleRate;
 
 			SetIsPreviewingLoopRegion(static_cast<double>(StartLoopPosInSeconds), static_cast<double>(EndLoopPosInSeconds), true);
 		}
@@ -830,6 +853,11 @@ void FWaveformTransformationMarkerRenderer::PreviewSelectedLoop()
 
 void FWaveformTransformationMarkerRenderer::SetIsPreviewingLoopRegion(double InStartTime, double InEndTime, bool bIsPreviewing)
 {
+	// Stop any current loop previewing so StartLoopTime is never greater than EndLoopTime (Other than when EndLoopTime is invalid)
+	IsPreviewingLoopRegionHandle->SetValue(false);
+	StartLoopTimeHandle->SetValue(0.0);
+	EndLoopTimeHandle->SetValue(-1.0);
+
 	IsPreviewingLoopRegionHandle->SetValue(bIsPreviewing);
 	StartLoopTimeHandle->SetValue(InStartTime);
 	EndLoopTimeHandle->SetValue(InEndTime);
@@ -929,8 +957,7 @@ void FWaveformTransformationMarkerRenderer::SetPropertyValueDependingOnInteracti
 			{
 				check(PixelsPerFrame > 0);
 				check(TransformationWaveInfo.NumChannels > 0);
-				const uint32 StartFrameOffset = TransformationWaveInfo.StartFrameOffset / TransformationWaveInfo.NumChannels;
-				MarkerInInteractionRange.SetMarkerPosition(((LocalCursorXClamped + (StartFrameOffset * PixelsPerFrame)) / PixelsPerFrame), false);
+				MarkerInInteractionRange.SetMarkerPosition((LocalCursorXClamped / PixelsPerFrame), false);
 				HighlightedCue = MarkerInInteractionRange.GetMarker()->CuePointID;
 			}
 		}
@@ -941,8 +968,7 @@ void FWaveformTransformationMarkerRenderer::SetPropertyValueDependingOnInteracti
 		{
 			check(PixelsPerFrame > 0);
 			check(TransformationWaveInfo.NumChannels > 0);
-			const uint32 StartFrameOffset = TransformationWaveInfo.StartFrameOffset / TransformationWaveInfo.NumChannels;
-			MarkerInInteractionRange.SetMarkerPosition(((LocalCursorXClamped + (StartFrameOffset * PixelsPerFrame)) / PixelsPerFrame), true);
+			MarkerInInteractionRange.SetMarkerPosition((LocalCursorXClamped / PixelsPerFrame), true);
 			HighlightedCue = MarkerInInteractionRange.GetMarker()->CuePointID;
 		}
 	}
@@ -1048,7 +1074,7 @@ void FWaveformTransformationMarkerRenderer::ShowConvertToLoopAtMarker(const FGeo
 		});
 }
 
-void FWaveformTransformationMarkerRenderer::FMarkerInInteractionRange::SetMarkerInInteractionRange(FSoundWaveCuePoint* InMarkerPtr, const bool InIsLoopRegionEndMarker, int32 GrabXOffset)
+void FWaveformTransformationMarkerRenderer::FMarkerInInteractionRange::SetMarkerInInteractionRange(FSoundWaveCuePoint* InMarkerPtr, const bool InIsLoopRegionEndMarker, int64 GrabXOffset)
 {
 	MarkerPtr = InMarkerPtr;
 	IsLoopRegionEndMarker = InIsLoopRegionEndMarker;
@@ -1064,18 +1090,20 @@ void FWaveformTransformationMarkerRenderer::FMarkerInInteractionRange::SetMarker
 		//FMath::Max prevents the user from collapsing the loop
 		if (bIsLeft)
 		{
-			int32 MoveDifference = MarkerPtr->FramePosition - (int32)InPosition;
-			MarkerPtr->FramePosition = (int32)InPosition;
-			MarkerPtr->FrameLength = FMath::Max(MarkerPtr->FrameLength + MoveDifference, MinLoopSize);
+			int64 MoveDifference = MarkerPtr->FramePosition - static_cast<int64>(InPosition);
+			int64 NewFrameLength = FMath::Max(MinLoopSize, MarkerPtr->FrameLength + MoveDifference);
+			
+			MarkerPtr->FramePosition = static_cast<int64>(InPosition);
+			MarkerPtr->FrameLength = NewFrameLength;
 		}
 		else
 		{
-			MarkerPtr->FrameLength = FMath::Max(InPosition, MarkerPtr->FramePosition + MinLoopSize) - MarkerPtr->FramePosition;
+			MarkerPtr->FrameLength = FMath::Max(static_cast<int64>(InPosition), MarkerPtr->FramePosition + MinLoopSize) - MarkerPtr->FramePosition;
 		}
 	}
 	else
 	{
-		MarkerPtr->FramePosition = FMath::Max(0, InPosition - GrabOffset);
+		MarkerPtr->FramePosition = FMath::Max(0, static_cast<int64>(InPosition) - GrabOffset);
 	}
 }
 
diff --git a/Engine/Plugins/Editor/WaveformEditor/Source/WaveformTransformationsWidgets/Private/WaveformTransformationMarkersRenderer.h b/Engine/Plugins/Editor/WaveformEditor/Source/WaveformTransformationsWidgets/Private/WaveformTransformationMarkersRenderer.h
--- a/Engine/Plugins/Editor/WaveformEditor/Source/WaveformTransformationsWidgets/Private/WaveformTransformationMarkersRenderer.h
+++ b/Engine/Plugins/Editor/WaveformEditor/Source/WaveformTransformationsWidgets/Private/WaveformTransformationMarkersRenderer.h
@@ -37,7 +37,7 @@ private:
 
 	void SetLoopColors(const TArray<FLinearColor>& InColors);
 
-	bool IsPositionInInteractionRange(const int32 InFramePosition, const FVector2D& InLocalCursorPosition, const FGeometry& WidgetGeometry) const;
+	bool IsPositionInInteractionRange(const int64 InFramePosition, const FVector2D& InLocalCursorPosition, const FGeometry& WidgetGeometry) const;
 	bool IsPositionInLoopArea(const FVector2D& InLocalCursorPosition, const FGeometry& WidgetGeometry, const FSoundWaveCuePoint* CueMarker, bool IncludeHandleArea = true) const;
 	bool IsPositionInLoopHandleArea(const FVector2D& InLocalCursorPosition, const FGeometry& WidgetGeometry, const FSoundWaveCuePoint& CueMarker, bool IsLeftHandle) const;
 	bool SetMarkerInInteractionRange(const FVector2D& InLocalCursorPosition, const FGeometry& WidgetGeometry);
@@ -62,17 +62,17 @@ private:
 	class FMarkerInInteractionRange
 	{
 	public:
-		void SetMarkerInInteractionRange(FSoundWaveCuePoint* InMarkerPtr, const bool InIsLoopRegionEndMarker = false, int32 GrabXOffset = 0);
+		void SetMarkerInInteractionRange(FSoundWaveCuePoint* InMarkerPtr, const bool InIsLoopRegionEndMarker = false, int64 GrabXOffset = 0);
 		FSoundWaveCuePoint* GetMarker() const { return MarkerPtr; }
 		void SetMarkerPosition(const float InPosition, bool bIsLeft);
-		void SetMinLoopSize(int32 value) { check(value > 0); MinLoopSize = value; }
-		const int32 GetMinLoopSize() { return MinLoopSize; }
+		void SetMinLoopSize(int64 value) { MinLoopSize = FMath::Max(UWaveCueArray::MinLoopSize, value); }
+		const int64 GetMinLoopSize() const { return MinLoopSize; }
 		
 	private:
 		FSoundWaveCuePoint* MarkerPtr = nullptr;
 		bool IsLoopRegionEndMarker = false;
-		int32 GrabOffset = 0;
-		int32 MinLoopSize = 10;
+		int64 GrabOffset = 0;
+		int64 MinLoopSize = UWaveCueArray::MinLoopSize;
 	} MarkerInInteractionRange;
 
 	EMarkerInteractionType EvaluateInteractionTypeFromCursorPosition(const FVector2D& InLocalCursorPosition, const FGeometry& WidgetGeometry, const FKey MouseButton);
diff --git a/Engine/Plugins/Editor/WaveformEditor/Source/WaveformTransformationsWidgets/Private/WaveformTransformationTrimFadeRenderer.cpp b/Engine/Plugins/Editor/WaveformEditor/Source/WaveformTransformationsWidgets/Private/WaveformTransformationTrimFadeRenderer.cpp
--- a/Engine/Plugins/Editor/WaveformEditor/Source/WaveformTransformationsWidgets/Private/WaveformTransformationTrimFadeRenderer.cpp
+++ b/Engine/Plugins/Editor/WaveformEditor/Source/WaveformTransformationsWidgets/Private/WaveformTransformationTrimFadeRenderer.cpp
@@ -239,17 +239,22 @@ FVector2D FWaveformTransformationTrimFadeRenderer::GetLocalCursorPosition(const
 
 double FWaveformTransformationTrimFadeRenderer::ConvertXRatioToTime(const float InRatio) const
 {
-	const float NumFrames = TransformationWaveInfo.NumAvilableSamples / TransformationWaveInfo.NumChannels;
+	check(TransformationWaveInfo.NumChannels > 0);
+
+	const float NumFrames = TransformationWaveInfo.TotalNumSamples / TransformationWaveInfo.NumChannels;
 	const float FrameSelected = NumFrames * InRatio;
 	return FrameSelected / TransformationWaveInfo.SampleRate;
 }
 
 void FWaveformTransformationTrimFadeRenderer::Tick(const FGeometry& AllottedGeometry, const double InCurrentTime, const float InDeltaTime)
 { 
-	const float NumFrames = TransformationWaveInfo.NumAvilableSamples / TransformationWaveInfo.NumChannels;
+	check(TransformationWaveInfo.NumChannels > 0);
+
+	const float NumFrames = TransformationWaveInfo.TotalNumSamples / TransformationWaveInfo.NumChannels;
 	const double FirstFrame = FMath::Clamp((GetPropertyValue<double>(StartTimeHandle) * TransformationWaveInfo.SampleRate) , 0.f, NumFrames);
 	const double EndFrame = FMath::Clamp((GetPropertyValue<double>(EndTimeHandle) * TransformationWaveInfo.SampleRate), FirstFrame, NumFrames);
 
+	check(NumFrames > 0);
 	PixelsPerFrame = FMath::Max(AllottedGeometry.GetLocalSize().X / NumFrames, 0.0);
 
 	StartTimeHandleX = FirstFrame * PixelsPerFrame;
diff --git a/Engine/Plugins/Editor/WaveformEditor/Source/WaveformTransformationsWidgets/Public/IWaveformTransformationRenderer.h b/Engine/Plugins/Editor/WaveformEditor/Source/WaveformTransformationsWidgets/Public/IWaveformTransformationRenderer.h
--- a/Engine/Plugins/Editor/WaveformEditor/Source/WaveformTransformationsWidgets/Public/IWaveformTransformationRenderer.h
+++ b/Engine/Plugins/Editor/WaveformEditor/Source/WaveformTransformationsWidgets/Public/IWaveformTransformationRenderer.h
@@ -22,8 +22,9 @@ struct FWaveformTransformationRenderInfo
 {
 	float SampleRate = 0.f;
 	int32 NumChannels = 0;
-	uint32 StartFrameOffset = 0;	//the first frame this transformation can act upon
-	uint32 NumAvilableSamples = 0;	//the number of samples available to the transformation
+	int64 StartFrameOffset = 0;	//the first frame this transformation can act upon
+	uint32 TotalNumSamples = 0;	//the number of total samples in the original waveform
+	uint32 NumSamplesAvailable = 0;	//the number of samples available to the transformation
 };
 
 class IWaveformTransformationRenderer
diff --git a/Engine/Source/Editor/AudioEditor/Private/Factories/SoundFactory.cpp b/Engine/Source/Editor/AudioEditor/Private/Factories/SoundFactory.cpp
--- a/Engine/Source/Editor/AudioEditor/Private/Factories/SoundFactory.cpp
+++ b/Engine/Source/Editor/AudioEditor/Private/Factories/SoundFactory.cpp
@@ -681,9 +681,9 @@ UObject* USoundFactory::CreateObject
 		for (FWaveCue& WaveCue : WaveInfo.WaveCues)
 		{
 			FSoundWaveCuePoint NewCuePoint;
-			NewCuePoint.CuePointID = (int32)WaveCue.CuePointID;
-			NewCuePoint.FrameLength = (int32)WaveCue.SampleLength;
-			NewCuePoint.FramePosition = (int32)WaveCue.Position;
+			NewCuePoint.CuePointID = static_cast<int32>(WaveCue.CuePointID);
+			NewCuePoint.FrameLength = static_cast<int64>(WaveCue.SampleLength);
+			NewCuePoint.FramePosition = static_cast<int64>(WaveCue.Position);
 			NewCuePoint.Label = WaveCue.Label;
 			Sound->CuePoints.Add(NewCuePoint);
 			WaveCuePointIDs.Add(NewCuePoint.CuePointID);
@@ -696,16 +696,17 @@ UObject* USoundFactory::CreateObject
 		{
 			FSoundWaveCuePoint NewCuePoint;
 			NewCuePoint.bIsLoopRegion = true;
-			NewCuePoint.CuePointID = (int32)SampleLoop.LoopID;
-			NewCuePoint.FramePosition = (int32)SampleLoop.StartFrame;
-			NewCuePoint.FrameLength = (int32)SampleLoop.EndFrame - (int32)SampleLoop.StartFrame;
+			NewCuePoint.CuePointID = static_cast<int32>(SampleLoop.LoopID);
+			NewCuePoint.FramePosition = static_cast<int64>(SampleLoop.StartFrame);
+			NewCuePoint.FrameLength = static_cast<int64>(SampleLoop.EndFrame) - static_cast<int64>(SampleLoop.StartFrame);
 			if (SampleLoop.EndFrame <= SampleLoop.StartFrame)
 			{
 				Warn->Logf(ELogVerbosity::Error, 
-					TEXT("Found invalid start and end frames when creating Cue Point from Sample Loop Region! LoopID = %d, StartFrame = %d, EndFrame = %d"), 
+					TEXT("EndFrame must be greater than StartFrame when creating Cue Point from Sample Loop Region! LoopID = %d, StartFrame = %d, EndFrame = %d"), 
 					SampleLoop.LoopID, SampleLoop.StartFrame, SampleLoop.EndFrame);
 
 				FoundInvalidSampleLoops = true;
+				NewCuePoint.FrameLength = 0;
 			}
 
 			if (WaveCuePointIDs.Contains(NewCuePoint.CuePointID))
diff --git a/Engine/Source/Runtime/AudioExtensions/Public/IWaveformTransformation.h b/Engine/Source/Runtime/AudioExtensions/Public/IWaveformTransformation.h
--- a/Engine/Source/Runtime/AudioExtensions/Public/IWaveformTransformation.h
+++ b/Engine/Source/Runtime/AudioExtensions/Public/IWaveformTransformation.h
@@ -73,12 +73,12 @@ struct FSoundWaveCuePoint
 	FString Label;
 
 	// The frame position of the cue point
-	UPROPERTY(Category = Info, EditAnywhere, BlueprintReadOnly)
-	int32 FramePosition = 0;
+	UPROPERTY(Category = Info, EditAnywhere, BlueprintReadOnly, meta = (ClampMin = 0))
+	int64 FramePosition = 0;
 
 	// The frame length of the cue point (non-zero if it's a region)
-	UPROPERTY(Category = Info, EditAnywhere, BlueprintReadOnly)
-	int32 FrameLength = 0;
+	UPROPERTY(Category = Info, EditAnywhere, BlueprintReadOnly, meta = (ClampMin = 0))
+	int64 FrameLength = 0;
 
 	bool IsLoopRegion() const { return bIsLoopRegion; }
 	void SetLoopRegion(bool value) { bIsLoopRegion = value; }
@@ -86,18 +86,18 @@ struct FSoundWaveCuePoint
 #if WITH_EDITORONLY_DATA
 	void ScaleFrameValues(float Factor)
 	{
-		FramePosition = FMath::FloorToInt((float)FramePosition * Factor);
-		FrameLength = FMath::FloorToInt((float)FrameLength * Factor);
+		FramePosition = FMath::FloorToInt64((float)FramePosition * Factor);
+		FrameLength = FMath::FloorToInt64((float)FrameLength * Factor);
 	}
 #endif // WITH_EDITORONLY_DATA
 
 	friend class USoundFactory;
 	friend class USoundWave;
 
-private:
+protected:
 	// intentionally kept private.
 	// only USoundFactory should modify this value on import
-	UPROPERTY(Category = Info, EditAnywhere)
+	UPROPERTY(Category = Info, EditAnywhere, BlueprintReadOnly)
 	bool bIsLoopRegion = false;
 };
 
