[Waveform Editor Bug] - Can add Multiple of the same Default Transformations but not in a Sound Wave

diff --git a/Engine/Plugins/Editor/WaveformEditor/Source/WaveformEditor/Private/WaveformEditorTransformationsSettings.h b/Engine/Plugins/Editor/WaveformEditor/Source/WaveformEditor/Private/WaveformEditorTransformationsSettings.h
--- a/Engine/Plugins/Editor/WaveformEditor/Source/WaveformEditor/Private/WaveformEditorTransformationsSettings.h
+++ b/Engine/Plugins/Editor/WaveformEditor/Source/WaveformEditor/Private/WaveformEditorTransformationsSettings.h
@@ -27,7 +27,11 @@ class UWaveformEditorTransformationsSettings : public UDeveloperSettings
 	//~ End UDeveloperSettings interface
 	
 public:
-	/** A Transformation chain that will be added to the inspected Soundwave if there aren't any  */
+	/** 
+	 * A Transformation chain that will be added to the inspected Soundwave if there aren't any  
+	 * Using TSet to limit each transformation type to one instance to reduce complexity and create a more consistent experience
+	 * USoundWave::Transformations is a TArray because changing it to TSet will delete user data
+	 */
 	UPROPERTY(config, EditAnywhere, Category = "Launch Options")
-	TArray<TSubclassOf<UWaveformTransformationBase>> LaunchTransformations;
+	TSet<TSubclassOf<UWaveformTransformationBase>> LaunchTransformations;
 };
diff --git a/Engine/Source/Runtime/Engine/Classes/Sound/SoundWave.h b/Engine/Source/Runtime/Engine/Classes/Sound/SoundWave.h
--- a/Engine/Source/Runtime/Engine/Classes/Sound/SoundWave.h
+++ b/Engine/Source/Runtime/Engine/Classes/Sound/SoundWave.h
@@ -1068,7 +1068,10 @@ public:
 #pragma endregion
 	} RawData;
 
-	/** Waveform edits to be applied to this SoundWave on cook (editing transformations will trigger a cook) */
+	/** 
+	 * Waveform edits to be applied to this SoundWave on cook (editing transformations will trigger a cook)
+	 * Transformations is a TArray because changing it to TSet will delete user data
+	 */
 	UPROPERTY(EditAnywhere, Instanced, Category = "Waveform Processing")
 	TArray<TObjectPtr<class UWaveformTransformationBase>> Transformations;
 #endif
