[Waveform Editor] - SoundWave markers in wrong place when exporting to Soundwaves.

diff --git a/Engine/Plugins/Editor/WaveformEditor/Source/WaveformTransformations/Private/WaveformTransformationMarkers.cpp b/Engine/Plugins/Editor/WaveformEditor/Source/WaveformTransformations/Private/WaveformTransformationMarkers.cpp
--- a/Engine/Plugins/Editor/WaveformEditor/Source/WaveformTransformations/Private/WaveformTransformationMarkers.cpp
+++ b/Engine/Plugins/Editor/WaveformEditor/Source/WaveformTransformations/Private/WaveformTransformationMarkers.cpp
@@ -98,6 +98,11 @@ void UWaveformTransformationMarkers::UpdateConfiguration(FWaveTransformUObjectCo
 	check(Markers != nullptr);
 	Markers->InitMarkersIfNotSet(InOutConfiguration.WaveCues);
 
+	// Assert that InOutConfiguration is initialized and valid
+	check(InOutConfiguration.SampleRate > 0);
+	check(InOutConfiguration.StartTime >= 0);
+	StartFrameOffset = static_cast<int64>(InOutConfiguration.StartTime * InOutConfiguration.SampleRate);
+
 	if (!bCachedIsPreviewingLoopRegion)
 	{
 		bCachedSoundWaveLoopState = InOutConfiguration.bCachedSoundWaveLoopState;
@@ -136,38 +141,37 @@ void UWaveformTransformationMarkers::CycleMarkerLoopRegion(ELoopModificationCont
 #if WITH_EDITOR
 void UWaveformTransformationMarkers::OverwriteSoundWaveData(USoundWave& InOutSoundWave)
 {
-	uint32 TransformationMarkerCount = 0;
-	uint32 ThisTransformationNumber = 0;
+	// Overwriting soundwave data can cause a change in number of samples, invalidating the FramePositions of CuesAndLoops
+	// Subtracting StartFrameOffset from FramePosition shifts the CuesAndLoops to the correct relative position
+	TArray<FSoundWaveCuePoint> FrameAdjustedCuesAndLoops = Markers->CuesAndLoops;
+
+	// Users can export a loop region, shift FramePositions relative to the Loop region
+	if (bIsPreviewingLoopRegion)
+	{
+		check(InOutSoundWave.GetImportedSampleRate() > 0);
+		int64 StartLoopFramePosition = StartLoopTime * InOutSoundWave.GetImportedSampleRate();
+		StartFrameOffset = StartFrameOffset < StartLoopFramePosition ? StartLoopFramePosition : StartFrameOffset;
+	}
 
-	for (TObjectPtr<UWaveformTransformationBase> Transformation : InOutSoundWave.Transformations)
+	if (StartFrameOffset != 0)
 	{
-		TObjectPtr<UWaveformTransformationMarkers> MarkersTransformation = Cast<UWaveformTransformationMarkers>(Transformation);
-		if (MarkersTransformation)
+		for (FSoundWaveCuePoint& Marker : FrameAdjustedCuesAndLoops)
 		{
-			TransformationMarkerCount++;
-			if (this == MarkersTransformation)
+			const int64 NewPosition = Marker.FramePosition - StartFrameOffset;
+
+			// If loop region is cut, resize it to maintain proper relative loop end point
+			// else maintain original loop region size
+			if (NewPosition < 0 && Marker.IsLoopRegion())
 			{
-				ThisTransformationNumber = TransformationMarkerCount;
-				break;
+				Marker.FrameLength = Marker.FrameLength + NewPosition > 0 ? Marker.FrameLength + NewPosition : Marker.FrameLength;
+				Marker.FrameLength = FMath::Max(Marker.FrameLength, UWaveCueArray::MinLoopSize);
 			}
+
+			Marker.FramePosition = FMath::Max(NewPosition, 0);
 		}
 	}
 
-	check(ThisTransformationNumber > 0);
-
-	// If this is the first Marker transformation, overwrite the soundwave original SoundWaveCuePoints
-	if (ThisTransformationNumber == 1)
-	{
-		InOutSoundWave.SetSoundWaveCuePoints(Markers->CuesAndLoops);
-	}
-	// Handles the case if there are multiple marker transformations to preserve all marker "tracks"
-	else
-	{
-		TArray<FSoundWaveCuePoint> SoundWaveCuePoints(InOutSoundWave.GetSoundWaveCuePoints());
-		SoundWaveCuePoints.Append(Markers->CuesAndLoops);
-		
-		InOutSoundWave.SetSoundWaveCuePoints(SoundWaveCuePoints);
-	}
+	InOutSoundWave.SetSoundWaveCuePoints(FrameAdjustedCuesAndLoops);
 }
 
 void UWaveformTransformationMarkers::GetTransformationInfo(FWaveformTransformationInfo& InOutTransformationInfo) const
diff --git a/Engine/Plugins/Editor/WaveformEditor/Source/WaveformTransformations/Public/WaveformTransformationMarkers.h b/Engine/Plugins/Editor/WaveformEditor/Source/WaveformTransformations/Public/WaveformTransformationMarkers.h
--- a/Engine/Plugins/Editor/WaveformEditor/Source/WaveformTransformations/Public/WaveformTransformationMarkers.h
+++ b/Engine/Plugins/Editor/WaveformEditor/Source/WaveformTransformations/Public/WaveformTransformationMarkers.h
@@ -79,14 +79,15 @@ public:
 	UPROPERTY(VisibleAnywhere, Category = "Markers")
 	TObjectPtr<UWaveCueArray> Markers;
 
-	UPROPERTY(EditAnywhere, Transient, Category = "Loop Preview", meta = (ClampMin = 0.0))
+	// These properties are hidden in editor by an IPropertyTypeCustomization
+	UPROPERTY(EditAnywhere, Category = "Loop Preview", meta = (ClampMin = 0.0))
 	double StartLoopTime = 0.0;
-	UPROPERTY(EditAnywhere, Transient, Category = "Loop Preview")
+	// When EndLoopTime is < 0 it skips processing the audio in FWaveTransformationMarkers::ProcessAudio 
+	UPROPERTY(EditAnywhere, Category = "Loop Preview")
 	double EndLoopTime = -1.0;
-	UPROPERTY(EditAnywhere, Transient, Category = "Loop Preview")
+	UPROPERTY(EditAnywhere, Category = "Loop Preview")
 	bool bIsPreviewingLoopRegion = false;
 
-	// Default value for EndLoopTime is < 0 to aid with proper initialization in UpdateDurationProperties()
 	
 	virtual Audio::FTransformationPtr CreateTransformation() const override;
 
@@ -105,6 +106,13 @@ public:
 private:
 	float AvailableWaveformDuration = -1.f;
 
+	/* Made a UPROPERTY so that it is captured when the soundwave is duplicated during export
+	 * Given these specifiers to prevent a crash in the property handle caching system used for transformations
+	 * Note: hidden in the details panel by the IPropertyTypeCustomization
+	 */
+	UPROPERTY(VisibleAnywhere, Category = "Markers")
+	int64 StartFrameOffset = 0;
+
 	bool bCachedIsPreviewingLoopRegion = false;
 	bool bCachedSoundWaveLoopState = false;
 };
