[Waveform Editor] - Loop preview region does not move with the loop region when properties are edited in the processing panel

diff --git a/Engine/Plugins/Editor/WaveformEditor/Source/WaveformTransformations/Private/WaveformTransformationMarkers.cpp b/Engine/Plugins/Editor/WaveformEditor/Source/WaveformTransformations/Private/WaveformTransformationMarkers.cpp
--- a/Engine/Plugins/Editor/WaveformEditor/Source/WaveformTransformations/Private/WaveformTransformationMarkers.cpp
+++ b/Engine/Plugins/Editor/WaveformEditor/Source/WaveformTransformations/Private/WaveformTransformationMarkers.cpp
@@ -88,6 +88,42 @@ UWaveformTransformationMarkers::UWaveformTransformationMarkers(const FObjectInit
 	}
 }
 
+#if WITH_EDITOR
+void UWaveformTransformationMarkers::PostEditChangeProperty(FPropertyChangedEvent& PropertyChangedEvent)
+{
+	Super::PostEditChangeProperty(PropertyChangedEvent);
+
+	const FName PropertyName = PropertyChangedEvent.GetPropertyName();
+	if (bIsPreviewingLoopRegion && PropertyName == GET_MEMBER_NAME_CHECKED(UWaveformTransformationMarkers, Markers))
+	{
+		check(Markers);
+		for (const FSoundWaveCuePoint& CuePoint : Markers->CuesAndLoops)
+		{
+			// Update LoopRegion Preview window if a loop is selected and the user edits the loop bounds using the properties window
+			if (Markers->SelectedCue == CuePoint.CuePointID)
+			{
+				check(CuePoint.IsLoopRegion());
+				check(SampleRate > 0.f);
+				check(EndLoopTime > 0)
+
+				const int64 StartLoopFramePos = StartLoopTime * SampleRate;
+				const int64 EndLoopFramePos = EndLoopTime * SampleRate;
+
+				if (CuePoint.FramePosition != StartLoopFramePos)
+				{
+					StartLoopTime = static_cast<float>(CuePoint.FramePosition) / SampleRate;
+				}
+
+				if (CuePoint.FramePosition + CuePoint.FrameLength != EndLoopFramePos)
+				{
+					EndLoopTime = static_cast<float>(CuePoint.FramePosition + CuePoint.FrameLength) / SampleRate;
+				}
+			}
+		}
+	}
+}
+#endif // WITH_EDITOR
+
 Audio::FTransformationPtr UWaveformTransformationMarkers::CreateTransformation() const
 {
 	return MakeUnique<FWaveTransformationMarkers>(StartLoopTime, EndLoopTime);
@@ -99,9 +135,10 @@ void UWaveformTransformationMarkers::UpdateConfiguration(FWaveTransformUObjectCo
 	Markers->InitMarkersIfNotSet(InOutConfiguration.WaveCues);
 
 	// Assert that InOutConfiguration is initialized and valid
-	check(InOutConfiguration.SampleRate > 0);
-	check(InOutConfiguration.StartTime >= 0);
+	check(InOutConfiguration.SampleRate > 0.f);
+	check(InOutConfiguration.StartTime >= 0.f);
 	StartFrameOffset = static_cast<int64>(InOutConfiguration.StartTime * InOutConfiguration.SampleRate);
+	SampleRate = InOutConfiguration.SampleRate;
 
 	if (!bCachedIsPreviewingLoopRegion)
 	{
diff --git a/Engine/Plugins/Editor/WaveformEditor/Source/WaveformTransformations/Public/WaveformTransformationMarkers.h b/Engine/Plugins/Editor/WaveformEditor/Source/WaveformTransformations/Public/WaveformTransformationMarkers.h
--- a/Engine/Plugins/Editor/WaveformEditor/Source/WaveformTransformations/Public/WaveformTransformationMarkers.h
+++ b/Engine/Plugins/Editor/WaveformEditor/Source/WaveformTransformations/Public/WaveformTransformationMarkers.h
@@ -76,6 +76,10 @@ class WAVEFORMTRANSFORMATIONS_API UWaveformTransformationMarkers : public UWavef
 public:
 	UWaveformTransformationMarkers(const FObjectInitializer& ObjectInitializer);
 
+#if WITH_EDITOR
+	virtual void PostEditChangeProperty(FPropertyChangedEvent& PropertyChangedEvent) override;
+#endif // WITH_EDITOR
+
 	UPROPERTY(VisibleAnywhere, Category = "Markers")
 	TObjectPtr<UWaveCueArray> Markers;
 
@@ -104,6 +108,7 @@ public:
 #endif //WITH_EDITOR
 
 private:
+	float SampleRate = 0.f;
 	float AvailableWaveformDuration = -1.f;
 
 	/* Made a UPROPERTY so that it is captured when the soundwave is duplicated during export
