Waveform Editor: Scrubber Undo Fixes

diff --git a/Engine/Plugins/Editor/WaveformEditor/Source/WaveformEditor/Private/WaveformEditor.cpp b/Engine/Plugins/Editor/WaveformEditor/Source/WaveformEditor/Private/WaveformEditor.cpp
--- a/Engine/Plugins/Editor/WaveformEditor/Source/WaveformEditor/Private/WaveformEditor.cpp
+++ b/Engine/Plugins/Editor/WaveformEditor/Source/WaveformEditor/Private/WaveformEditor.cpp
@@ -139,6 +139,13 @@ FWaveformEditor::~FWaveformEditor()
 		SoundWave->bLooping = TransformationChainConfig.bCachedSoundWaveLoopState;
 		SoundWave->PostEditChange();
 	}
+
+	TObjectPtr<UWaveformTransformationMarkers> MarkerTransformation = GetMarkerTransformation();
+
+	if (MarkerTransformation != nullptr && bCueChangeRegisteredByWaveformEditor)
+	{
+		MarkerTransformation->Markers->CueChanged.Unbind();
+	}
 }
 
 void FWaveformEditor::AddDefaultTransformations()
@@ -240,6 +247,18 @@ TObjectPtr<UWaveformTransformationMarkers> FWaveformEditor::GetOrAddMarkerTransf
 {
 	check(SoundWave);
 
+	TObjectPtr<UWaveformTransformationMarkers> MarkersTransformation = GetMarkerTransformation();
+
+	if (MarkersTransformation == nullptr)
+	{
+		MarkersTransformation = AddMarkerTransformation();
+	}
+
+	return MarkersTransformation;
+}
+
+TObjectPtr<UWaveformTransformationMarkers> FWaveformEditor::GetMarkerTransformation()
+{
 	TObjectPtr<UWaveformTransformationMarkers> MarkersTransformation;
 
 	for (TObjectPtr<UWaveformTransformationBase> Transformation : SoundWave->Transformations)
@@ -247,31 +266,33 @@ TObjectPtr<UWaveformTransformationMarkers> FWaveformEditor::GetOrAddMarkerTransf
 		MarkersTransformation = Cast<UWaveformTransformationMarkers>(Transformation);
 		if (MarkersTransformation)
 		{
-			break;
+			return MarkersTransformation;
 		}
 	}
 
-	if (MarkersTransformation == nullptr)
-	{
-		EObjectFlags MaskedOuterFlags = SoundWave ? SoundWave->GetMaskedFlags(RF_PropagateToSubObjects) : RF_NoFlags;
+	return nullptr;
+}
 
-		if (SoundWave->HasAnyFlags(RF_ClassDefaultObject | RF_ArchetypeObject))
-		{
-			MaskedOuterFlags |= RF_ArchetypeObject;
-		}
+TObjectPtr<UWaveformTransformationMarkers> FWaveformEditor::AddMarkerTransformation()
+{
+	TObjectPtr<UWaveformTransformationMarkers> MarkersTransformation;
+	EObjectFlags MaskedOuterFlags = SoundWave ? SoundWave->GetMaskedFlags(RF_PropagateToSubObjects) : RF_NoFlags;
 
-		MarkersTransformation = NewObject<UWaveformTransformationMarkers>(SoundWave, UWaveformTransformationMarkers::StaticClass(), NAME_None, MaskedOuterFlags);
+	if (SoundWave->HasAnyFlags(RF_ClassDefaultObject | RF_ArchetypeObject))
+	{
+		MaskedOuterFlags |= RF_ArchetypeObject;
+	}
 
-		SoundWave->Modify();
-		SoundWave->Transformations.Add(MarkersTransformation);
+	MarkersTransformation = NewObject<UWaveformTransformationMarkers>(SoundWave, UWaveformTransformationMarkers::StaticClass(), NAME_None, MaskedOuterFlags);
 
-		NotifyPostTransformationChange();
-	}
+	SoundWave->Modify();
+	SoundWave->Transformations.Add(MarkersTransformation);
+
+	NotifyPostTransformationChange();
 
 	return MarkersTransformation;
 }
 
-
 void FWaveformEditor::ToggleFadeIn()
 {
 	check(SoundWave);
@@ -476,6 +497,15 @@ void FWaveformEditor::SkipToNextMarker()
 	}
 }
 
+void FWaveformEditor::RegenerateTransformations()
+{
+	WaveformView.DataProvider->GenerateLayersChain();
+	WaveformView.DataProvider->UpdateRenderElements();
+	TransformationChainConfig = SoundWave->GetTransformationChainConfig();
+
+	UpdateTransportState();
+}
+
 void FWaveformEditor::ModifyMarkerLoopRegion(ELoopModificationControls Modification)
 {
 	check(SoundWave);
@@ -1277,61 +1307,87 @@ void FWaveformEditor::NotifyPostChange(const FPropertyChangedEvent& PropertyChan
 
 	if (PropertyChangedEvent.ChangeType == EPropertyChangeType::ValueSet)
 	{
-		if (!TransportController->IsPlaying())
-		{
-			// StartTime=0, EndTime=-1 if no transformations
-			const float TransformationDuration = TransformationChainConfig.EndTime - TransformationChainConfig.StartTime;
+		UpdateTransportState();
 
-			// Clamp Playback/Start times to start if out of range of TransformationChainConfig.
-			// If no active+initialized transformation, TransformationDuration will be <= 0.0f, ignore in this case.
-			if ( (PlaybackTimeBeforeTransformInteraction + StartTimeBeforeTransformInteraction - TransformationChainConfig.StartTime <= 0)
-			  || (TransformationDuration > 0.0f && PlaybackTimeBeforeTransformInteraction + StartTimeBeforeTransformInteraction - TransformationChainConfig.StartTime >= TransformationDuration) )
-			{
-				PlaybackTimeBeforeTransformInteraction = 0;
-				StartTimeBeforeTransformInteraction = TransformationChainConfig.StartTime;
-			}
-			const float StartTimeDifference = StartTimeBeforeTransformInteraction - TransformationChainConfig.StartTime;
-			const float AdjustedPlaybackTime = PlaybackTimeBeforeTransformInteraction + StartTimeDifference;
-			
-			switch (TransformInteractionPlayState)
-			{
-			case EAudioComponentPlayState::Playing:
-				TransportController->Play(AdjustedPlaybackTime);
-				TransformInteractionPlayState = EAudioComponentPlayState::Stopped;
-				break;
-			case EAudioComponentPlayState::Paused:
-				TransformInteractionPlayState = EAudioComponentPlayState::Stopped;
-			case EAudioComponentPlayState::Stopped:
-				// Set TransportController and displayed TransportCoordinator to value clamped within TransformationChainConfig.
-				TransportController->CacheStartTime(AdjustedPlaybackTime);
-				if (!TransportCoordinator->IsScrubbing())
-				{
-					const float CachedAudioStartTimeAsPercentage = TransportController->GetCachedAudioStartTimeAsPercentage();
-					TransportCoordinator->SetProgressRatio(CachedAudioStartTimeAsPercentage);
-				}
-				break;
-			default:
-				break;
-			}
+		bIsInteractingWithTransformations = false;
+	}
+	TObjectPtr<UWaveformTransformationMarkers> MarkerTransformation = GetMarkerTransformation();
+
+	if (MarkerTransformation != nullptr && !bCueChangeRegisteredByWaveformEditor)
+	{
+		MarkerTransformation->Markers->CueChanged.BindLambda([this]() { RegenerateTransformations(); });
+		bCueChangeRegisteredByWaveformEditor = true;
+	}
+	else if (MarkerTransformation == nullptr)
+	{
+		bCueChangeRegisteredByWaveformEditor = false;
+	}
+}
 
-			if (bWasPlayingBeforeChange)
+void FWaveformEditor::UpdateTransportState()
+{
+	if (!TransportController->IsPlaying())
+	{
+		// StartTime=0, EndTime=-1 if no transformations
+		const float TransformationDuration = TransformationChainConfig.EndTime - TransformationChainConfig.StartTime;
+
+		// Clamp Playback/Start times to start if out of range of TransformationChainConfig.
+		// If no active+initialized transformation, TransformationDuration will be <= 0.0f, ignore in this case.
+		if ((PlaybackTimeBeforeTransformInteraction + StartTimeBeforeTransformInteraction - TransformationChainConfig.StartTime <= 0)
+			|| (TransformationDuration > 0.0f && PlaybackTimeBeforeTransformInteraction + StartTimeBeforeTransformInteraction - TransformationChainConfig.StartTime >= TransformationDuration))
+		{
+			PlaybackTimeBeforeTransformInteraction = 0;
+			StartTimeBeforeTransformInteraction = TransformationChainConfig.StartTime;
+		}
+		const float StartTimeDifference = StartTimeBeforeTransformInteraction - TransformationChainConfig.StartTime;
+		const float AdjustedPlaybackTime = PlaybackTimeBeforeTransformInteraction + StartTimeDifference;
+
+		switch (TransformInteractionPlayState)
+		{
+		case EAudioComponentPlayState::Playing:
+			TransportController->Play(AdjustedPlaybackTime);
+			TransformInteractionPlayState = EAudioComponentPlayState::Stopped;
+			break;
+		case EAudioComponentPlayState::Paused:
+			TransformInteractionPlayState = EAudioComponentPlayState::Stopped;
+		case EAudioComponentPlayState::Stopped:
+			// Set TransportController and displayed TransportCoordinator to value clamped within TransformationChainConfig.
+			TransportController->CacheStartTime(AdjustedPlaybackTime);
+			if (!TransportCoordinator->IsScrubbing())
 			{
-				AudioComponent->Play();
+				const float CachedAudioStartTimeAsPercentage = TransportController->GetCachedAudioStartTimeAsPercentage();
+				TransportCoordinator->SetProgressRatio(CachedAudioStartTimeAsPercentage);
 			}
+			break;
+		default:
+			break;
 		}
 
-		bIsInteractingWithTransformations = false;
+		if (bWasPlayingBeforeChange)
+		{
+			AudioComponent->Play();
+		}
 	}
 }
 
-
-
 void FWaveformEditor::PostUndo(bool bSuccess)
 {
 	if (bSuccess)
 	{
 		WaveformView.DataProvider->GenerateLayersChain();
 		WaveformView.DataProvider->UpdateRenderElements();
+
+		TObjectPtr<UWaveformTransformationMarkers> MarkerTransformation = GetMarkerTransformation();
+
+		if (MarkerTransformation != nullptr && !bCueChangeRegisteredByWaveformEditor)
+		{
+			MarkerTransformation->Markers->CueChanged.BindLambda([this]() { RegenerateTransformations(); });
+			bCueChangeRegisteredByWaveformEditor = true;
+		}
+		else if (MarkerTransformation == nullptr)
+		{
+			bCueChangeRegisteredByWaveformEditor = false;
+		}
 	}
 }
 
diff --git a/Engine/Plugins/Editor/WaveformEditor/Source/WaveformEditor/Public/WaveformEditor.h b/Engine/Plugins/Editor/WaveformEditor/Source/WaveformEditor/Public/WaveformEditor.h
--- a/Engine/Plugins/Editor/WaveformEditor/Source/WaveformEditor/Public/WaveformEditor.h
+++ b/Engine/Plugins/Editor/WaveformEditor/Source/WaveformEditor/Public/WaveformEditor.h
@@ -107,6 +107,8 @@ private:
 	TObjectPtr<UWaveformTransformationTrimFade> GetOrAddTrimFadeTransformation();
 
 	TObjectPtr<UWaveformTransformationMarkers> GetOrAddMarkerTransformation();
+	TObjectPtr<UWaveformTransformationMarkers> GetMarkerTransformation();
+	TObjectPtr<UWaveformTransformationMarkers> AddMarkerTransformation();
 	void ToggleFadeIn();
 	bool CanFadeIn();
 
@@ -115,6 +117,8 @@ private:
 	void CreateMarker(bool bIsLoopRegion);
 	void DeleteMarker();
 	void SkipToNextMarker(); // Skips the playhead to the next marker
+	void RegenerateTransformations();
+	void UpdateTransportState();
 
 	void ModifyMarkerLoopRegion(ELoopModificationControls Modification);
 	void CycleMarkerLoopRegion(ELoopModificationControls Modification);
@@ -167,6 +171,10 @@ private:
 	float CachedFadeOutAmount = 0.f;
 	float FadeOutCurve = UWaveformTransformationTrimFade::FadeModeToCurveValueMap[EWaveEditorFadeMode::Exponetial];
 
+	// Tracking bool to ensure we detect wave cue changes and regenerate transformations
+	// Checking the event is bound may have issues if other systems need to subscribe for wave cue array changes
+	bool bCueChangeRegisteredByWaveformEditor = false;
+
 	enum class EWaveEditorReimportMode : uint8
 	{
 		SameFile = 0,
diff --git a/Engine/Plugins/Editor/WaveformEditor/Source/WaveformTransformations/Private/WaveformTransformationMarkers.cpp b/Engine/Plugins/Editor/WaveformEditor/Source/WaveformTransformations/Private/WaveformTransformationMarkers.cpp
--- a/Engine/Plugins/Editor/WaveformEditor/Source/WaveformTransformations/Private/WaveformTransformationMarkers.cpp
+++ b/Engine/Plugins/Editor/WaveformEditor/Source/WaveformTransformations/Private/WaveformTransformationMarkers.cpp
@@ -272,6 +272,8 @@ void UWaveCueArray::PostEditChangeProperty(FPropertyChangedEvent& PropertyChange
 			}
 		}
 	}
+
+	CueChanged.ExecuteIfBound();
 }
 
 void UWaveCueArray::InitMarkersIfNotSet(const TArray<FSoundWaveCuePoint>& InMarkers)
diff --git a/Engine/Plugins/Editor/WaveformEditor/Source/WaveformTransformations/Public/WaveformTransformationMarkers.h b/Engine/Plugins/Editor/WaveformEditor/Source/WaveformTransformations/Public/WaveformTransformationMarkers.h
--- a/Engine/Plugins/Editor/WaveformEditor/Source/WaveformTransformations/Public/WaveformTransformationMarkers.h
+++ b/Engine/Plugins/Editor/WaveformEditor/Source/WaveformTransformations/Public/WaveformTransformationMarkers.h
@@ -51,6 +51,9 @@ public:
 	// To minimize complexity while supporting all common editing cases, loops have a min length of 10 frames
 	static constexpr int64 MinLoopSize = 10;
 
+	DECLARE_DELEGATE(OnCueChange);
+	OnCueChange CueChanged;
+
 private:
 	UPROPERTY()
 	bool bIsInitialized = false;
