[Waveform editor] Add icons for fade options

diff --git a/Engine/Plugins/Editor/WaveformEditor/Source/WaveformEditor/Private/WaveformEditor.cpp b/Engine/Plugins/Editor/WaveformEditor/Source/WaveformEditor/Private/WaveformEditor.cpp
--- a/Engine/Plugins/Editor/WaveformEditor/Source/WaveformEditor/Private/WaveformEditor.cpp
+++ b/Engine/Plugins/Editor/WaveformEditor/Source/WaveformEditor/Private/WaveformEditor.cpp
@@ -3,6 +3,7 @@
 #include "WaveformEditor.h"
 
 #include "AudioDevice.h"
+#include "AudioWidgetsStyle.h"
 #include "AssetDefinitionRegistry.h"
 #include "Components/AudioComponent.h"
 #include "EditorReimportHandler.h"
@@ -811,16 +812,6 @@ bool FWaveformEditor::RegisterToolbar()
 		FToolMenuInsert InsertAfterExportSection("Export Controls", EToolMenuInsertType::After);
 		FToolMenuSection& TransformationSection = ToolBar->AddSection("Transformation Controls", TAttribute<FText>(), InsertAfterExportSection);
 
-		const FSlateBrush* IconBrush = FAppStyle::Get().GetBrush("Sequencer.Tracks.Fade");
-		FSlateBrush* FlippedIconBrush = new FSlateBrush(*IconBrush);
-		FlippedIconBrush->Mirroring = ESlateBrushMirrorType::Horizontal;
-
-		// Register the flipped brush in a custom style set
-		static TSharedRef<FSlateStyleSet> FlippedStyleSet = MakeShared<FSlateStyleSet>("FlippedStyleSet");
-		FlippedStyleSet->SetContentRoot(FAppStyle::Get().GetContentRootDir());
-		FlippedStyleSet->Set("Sequencer.Tracks.Fade.Flipped", FlippedIconBrush);
-		FSlateStyleRegistry::RegisterSlateStyle(*FlippedStyleSet);
-
 		TransformationSection.AddDynamicEntry("ToggleFadeInButton", FNewToolMenuSectionDelegate::CreateLambda([this, Commands](FToolMenuSection& InSection)
 			{
 				const UWaveformEditorToolMenuContext* Context = InSection.FindContext<UWaveformEditorToolMenuContext>();
@@ -831,7 +822,7 @@ bool FWaveformEditor::RegisterToolbar()
 						Commands.ToggleFadeIn,
 						FText(),
 						LOCTEXT("WaveformEditorFadeInButtonTooltip", "Toggle Fade In Transformation onto the soundwave"),
-						FSlateIcon("FlippedStyleSet", "Sequencer.Tracks.Fade.Flipped")
+						FSlateIcon(FAudioWidgetsStyle::StyleName, "AudioWidgetsStyle.FadeIn")
 					);
 					InSection.AddEntry(ToggleFadeIn);
 				}
@@ -867,7 +858,7 @@ bool FWaveformEditor::RegisterToolbar()
 						Commands.ToggleFadeOut,
 						FText(),
 						LOCTEXT("WaveformEditorFadeOutButtonTooltip", "Toggle Fade Out Transformation onto the soundwave"),
-						FSlateIcon(FAppStyle::GetAppStyleSetName(), "Sequencer.Tracks.Fade")
+						FSlateIcon(FAudioWidgetsStyle::StyleName, "AudioWidgetsStyle.FadeOut")
 					);
 					InSection.AddEntry(ToggleFadeOut);
 				}
@@ -903,7 +894,7 @@ bool FWaveformEditor::RegisterToolbar()
 						Commands.CreateMarker,
 						FText(),
 						LOCTEXT("WaveformEditorCreateMarkerButtonTooltip", "Create a marker cue for the soundwave"),
-						FSlateIcon(FAppStyle::GetAppStyleSetName(), "Sequencer.Tracks.Event")
+						FSlateIcon(FAudioWidgetsStyle::StyleName, "AudioWidgetsStyle.MarkerAdd")
 					);
 					InSection.AddEntry(CreateMarker);
 				}
@@ -920,7 +911,7 @@ bool FWaveformEditor::RegisterToolbar()
 						Commands.CreateLoopRegion,
 						FText(),
 						LOCTEXT("WaveformEditorCreateLoopRegionButtonTooltip", "Create a marker loop region for the soundwave"),
-						FSlateIcon(FAppStyle::GetAppStyleSetName(), "Graph.Node.Loop")
+						FSlateIcon(FAppStyle::GetAppStyleSetName(), "Sequencer.Transport.Looping")
 					);
 					InSection.AddEntry(CreateLoopRegion);
 				}
@@ -936,7 +927,7 @@ bool FWaveformEditor::RegisterToolbar()
 						Commands.DeleteMarker,
 						FText(),
 						LOCTEXT("WaveformEditorDeleteMarkerButtonTooltip", "Delete a marker cue or loop region for the soundwave"),
-						FSlateIcon(FAppStyle::GetAppStyleSetName(), "Symbols.X")
+						FSlateIcon(FAudioWidgetsStyle::StyleName, "AudioWidgetsStyle.MarkerDelete")
 					);
 					InSection.AddEntry(DeleteLoopRegion);
 				}
@@ -953,10 +944,10 @@ TSharedRef<SWidget> FWaveformEditor::GenerateFadeInOptionsMenu()
 
 	MenuBuilder.BeginSection(NAME_None, LOCTEXT("FadeInOptionsSection_Label", "Fade In Options"));
 	{
-		MenuBuilder.AddMenuEntry(FWaveformEditorCommands::Get().FadeInLinear);
-		MenuBuilder.AddMenuEntry(FWaveformEditorCommands::Get().FadeInExponential);
-		MenuBuilder.AddMenuEntry(FWaveformEditorCommands::Get().FadeInLogarithmic);
-		MenuBuilder.AddMenuEntry(FWaveformEditorCommands::Get().FadeInSigmoid);
+		MenuBuilder.AddMenuEntry(FWaveformEditorCommands::Get().FadeInLinear, NAME_None, TAttribute<FText>(), TAttribute<FText>(), FSlateIcon(FAudioWidgetsStyle::StyleName, "AudioWidgetsStyle.FadeInLinear"));
+		MenuBuilder.AddMenuEntry(FWaveformEditorCommands::Get().FadeInExponential, NAME_None, TAttribute<FText>(), TAttribute<FText>(), FSlateIcon(FAudioWidgetsStyle::StyleName, "AudioWidgetsStyle.FadeInExponential"));
+		MenuBuilder.AddMenuEntry(FWaveformEditorCommands::Get().FadeInLogarithmic, NAME_None, TAttribute<FText>(), TAttribute<FText>(), FSlateIcon(FAudioWidgetsStyle::StyleName, "AudioWidgetsStyle.FadeInLogarithmic"));
+		MenuBuilder.AddMenuEntry(FWaveformEditorCommands::Get().FadeInSigmoid, NAME_None, TAttribute<FText>(), TAttribute<FText>(), FSlateIcon(FAudioWidgetsStyle::StyleName, "AudioWidgetsStyle.FadeInSigmoid"));
 	}
 	MenuBuilder.EndSection();
 
@@ -970,10 +961,10 @@ TSharedRef<SWidget> FWaveformEditor::GenerateFadeOutOptionsMenu()
 
 	MenuBuilder.BeginSection(NAME_None, LOCTEXT("FadeOutOptionsSection_Label", "Fade Out Options"));
 	{
-		MenuBuilder.AddMenuEntry(FWaveformEditorCommands::Get().FadeOutLinear);
-		MenuBuilder.AddMenuEntry(FWaveformEditorCommands::Get().FadeOutExponential);
-		MenuBuilder.AddMenuEntry(FWaveformEditorCommands::Get().FadeOutLogarithmic);
-		MenuBuilder.AddMenuEntry(FWaveformEditorCommands::Get().FadeOutSigmoid);
+		MenuBuilder.AddMenuEntry(FWaveformEditorCommands::Get().FadeOutLinear, NAME_None, TAttribute<FText>(), TAttribute<FText>(), FSlateIcon(FAudioWidgetsStyle::StyleName, "AudioWidgetsStyle.FadeOutLinear"));
+		MenuBuilder.AddMenuEntry(FWaveformEditorCommands::Get().FadeOutExponential, NAME_None, TAttribute<FText>(), TAttribute<FText>(), FSlateIcon(FAudioWidgetsStyle::StyleName, "AudioWidgetsStyle.FadeOutExponential"));
+		MenuBuilder.AddMenuEntry(FWaveformEditorCommands::Get().FadeOutLogarithmic, NAME_None, TAttribute<FText>(), TAttribute<FText>(), FSlateIcon(FAudioWidgetsStyle::StyleName, "AudioWidgetsStyle.FadeOutLogarithmic"));
+		MenuBuilder.AddMenuEntry(FWaveformEditorCommands::Get().FadeOutSigmoid, NAME_None, TAttribute<FText>(), TAttribute<FText>(), FSlateIcon(FAudioWidgetsStyle::StyleName, "AudioWidgetsStyle.FadeOutSigmoid"));
 	}
 	MenuBuilder.EndSection();
 
diff --git a/Engine/Plugins/Editor/WaveformEditor/Source/WaveformTransformationsWidgets/Private/WaveformTransformationTrimFadeRenderer.cpp b/Engine/Plugins/Editor/WaveformEditor/Source/WaveformTransformationsWidgets/Private/WaveformTransformationTrimFadeRenderer.cpp
--- a/Engine/Plugins/Editor/WaveformEditor/Source/WaveformTransformationsWidgets/Private/WaveformTransformationTrimFadeRenderer.cpp
+++ b/Engine/Plugins/Editor/WaveformEditor/Source/WaveformTransformationsWidgets/Private/WaveformTransformationTrimFadeRenderer.cpp
@@ -2,14 +2,17 @@
 
 #include "WaveformTransformationTrimFadeRenderer.h"
 
+#include "AudioWidgetsStyle.h"
 #include "Brushes/SlateRoundedBoxBrush.h"
 #include "Framework/Application/SlateApplication.h"
 #include "Layout/Geometry.h"
 #include "PropertyHandle.h"
 #include "Rendering/DrawElements.h"
 #include "WaveformTransformationTrimFade.h"
+#include "Widgets/Images/SImage.h"
 #include "Widgets/Input/SButton.h"
 #include "Widgets/SWidget.h"
+#include "Widgets/Text/STextBlock.h"
 
 FWaveformTransformationTrimFadeRenderer::~FWaveformTransformationTrimFadeRenderer()
 {
@@ -367,6 +370,51 @@ void FWaveformTransformationTrimFadeRenderer::ShowSelectFadeModeMenuAtCursor(con
 		return;
 	}
 
+	TWeakPtr<IPropertyHandle> WeakStartFadeCurveHandle = StartFadeCurveHandle;
+	TWeakPtr<IPropertyHandle> WeakEndFadeCurveHandle = EndFadeCurveHandle;
+
+	TFunction<FSlateIcon(const EWaveEditorFadeMode&, const TSharedPtr<IPropertyHandle>&)> GetFadeModeIcon = 
+		[WeakStartFadeCurveHandle, WeakEndFadeCurveHandle](const EWaveEditorFadeMode& FadeMode, const TSharedPtr<IPropertyHandle>& PropertyHandle)
+		{
+			TSharedPtr<IPropertyHandle> LockedStartFadeCurveHandle = WeakStartFadeCurveHandle.Pin();
+			TSharedPtr<IPropertyHandle> LockedEndFadeCurveHandle = WeakEndFadeCurveHandle.Pin();
+			if (LockedStartFadeCurveHandle && PropertyHandle == LockedStartFadeCurveHandle)
+			{
+				switch (FadeMode)
+				{
+				case EWaveEditorFadeMode::Linear:
+					return FSlateIcon(FAudioWidgetsStyle::StyleName, "AudioWidgetsStyle.FadeInLinear");
+				case EWaveEditorFadeMode::Exponetial:
+					return FSlateIcon(FAudioWidgetsStyle::StyleName, "AudioWidgetsStyle.FadeInExponential");
+				case EWaveEditorFadeMode::Logarithmic:
+					return FSlateIcon(FAudioWidgetsStyle::StyleName, "AudioWidgetsStyle.FadeInLogarithmic");
+				case EWaveEditorFadeMode::Sigmoid:
+					return FSlateIcon(FAudioWidgetsStyle::StyleName, "AudioWidgetsStyle.FadeInSigmoid");
+				default:
+					return FSlateIcon();
+				}
+			}
+			else if (LockedEndFadeCurveHandle && PropertyHandle == LockedEndFadeCurveHandle)
+			{
+				switch (FadeMode)
+				{
+				case EWaveEditorFadeMode::Linear:
+					return FSlateIcon(FAudioWidgetsStyle::StyleName, "AudioWidgetsStyle.FadeOutLinear");
+				case EWaveEditorFadeMode::Exponetial:
+					return FSlateIcon(FAudioWidgetsStyle::StyleName, "AudioWidgetsStyle.FadeOutExponential");
+				case EWaveEditorFadeMode::Logarithmic:
+					return FSlateIcon(FAudioWidgetsStyle::StyleName, "AudioWidgetsStyle.FadeOutLogarithmic");
+				case EWaveEditorFadeMode::Sigmoid:
+					return FSlateIcon(FAudioWidgetsStyle::StyleName, "AudioWidgetsStyle.FadeOutSigmoid");
+				default:
+					return FSlateIcon();
+				}
+			}
+
+			return FSlateIcon();
+			
+		};
+
 	const FVector2D LocalCursorPosition = MouseEvent.GetScreenSpacePosition();
 
 	TSharedRef<SVerticalBox> MenuContent = SNew(SVerticalBox);
@@ -378,17 +426,36 @@ void FWaveformTransformationTrimFadeRenderer::ShowSelectFadeModeMenuAtCursor(con
 			.AutoHeight()
 			[
 				SNew(SButton)
-					.Text(FText::FromString(StaticEnum<EWaveEditorFadeMode>()->GetNameStringByValue(static_cast<int64>(FadeOptionPair.Key))))
-					.OnClicked_Lambda([this, PropertyHandle, FadeOptionPair]() -> FReply {
-					PropertyHandle->SetValue(FadeOptionPair.Value);
-
-					if (TSharedPtr<SWindow> LockedFadeModeMenuWindow = FadeModeMenuWindow.Pin())
+					.OnClicked_Lambda([this, PropertyHandle, FadeOptionPair]() -> FReply 
 					{
-						FSlateApplication::Get().RequestDestroyWindow(LockedFadeModeMenuWindow.ToSharedRef());
-					}
-
-					return FReply::Handled();
-				})
+						PropertyHandle->SetValue(FadeOptionPair.Value);
+
+						if (TSharedPtr<SWindow> LockedFadeModeMenuWindow = FadeModeMenuWindow.Pin())
+						{
+							FSlateApplication::Get().RequestDestroyWindow(LockedFadeModeMenuWindow.ToSharedRef());
+						}
+
+						return FReply::Handled();
+					})
+					[
+						SNew(SHorizontalBox)
+						+ SHorizontalBox::Slot()
+						.AutoWidth()
+						.VAlign(VAlign_Center)
+						.HAlign(HAlign_Left)
+						[
+							SNew(SImage)
+							.Image(GetFadeModeIcon(FadeOptionPair.Key, PropertyHandle).GetIcon())
+						]
+						+ SHorizontalBox::Slot()
+						.VAlign(VAlign_Center)
+						.HAlign(HAlign_Center)
+						.Padding(5)
+						[
+							SNew(STextBlock)
+							.Text(FText::FromString(StaticEnum<EWaveEditorFadeMode>()->GetNameStringByValue(static_cast<int64>(FadeOptionPair.Key))))
+						]
+					]
 			];
 	}
 
diff --git a/Engine/Plugins/Editor/WaveformEditor/Source/WaveformTransformationsWidgets/WaveformTransformationsWidgets.build.cs b/Engine/Plugins/Editor/WaveformEditor/Source/WaveformTransformationsWidgets/WaveformTransformationsWidgets.build.cs
--- a/Engine/Plugins/Editor/WaveformEditor/Source/WaveformTransformationsWidgets/WaveformTransformationsWidgets.build.cs
+++ b/Engine/Plugins/Editor/WaveformEditor/Source/WaveformTransformationsWidgets/WaveformTransformationsWidgets.build.cs
@@ -13,6 +13,7 @@ namespace UnrealBuildTool.Rules
 			PublicDependencyModuleNames.AddRange(
 				new string[] {
 					"AudioExtensions",
+					"AudioWidgets",
 					"ContentBrowser",
 					"Core",
 					"CoreUObject",
diff --git a/Engine/Plugins/Runtime/AudioWidgets/Source/AudioWidgets/Private/AudioWidgetsStyle.cpp b/Engine/Plugins/Runtime/AudioWidgets/Source/AudioWidgets/Private/AudioWidgetsStyle.cpp
--- a/Engine/Plugins/Runtime/AudioWidgets/Source/AudioWidgets/Private/AudioWidgetsStyle.cpp
+++ b/Engine/Plugins/Runtime/AudioWidgets/Source/AudioWidgets/Private/AudioWidgetsStyle.cpp
@@ -15,6 +15,7 @@
 #include "Misc/Paths.h"
 #include "Styling/AppStyle.h"
 #include "Styling/StyleColors.h"
+#include "Styling/SlateStyleMacros.h"
 #include "Styling/SlateStyleRegistry.h"
 
 FName FAudioWidgetsStyle::StyleName("AudioWidgetsStyle");
@@ -30,6 +31,7 @@ FAudioWidgetsStyle::FAudioWidgetsStyle()
 	SetParentStyleName(FAppStyle::GetAppStyleSetName());
 	SetContentRoot(FPaths::EnginePluginsDir() / TEXT("Runtime/AudioWidgets/Content"));
 	SetResources();
+	SetupStyles();
 
 	/** 
 	* AudioTextBox Style
@@ -180,6 +182,27 @@ void FAudioWidgetsStyle::SetResources()
 	Set(AudioWidgetsStylePrivate::ScrubHandleBrushName, new FSlateBoxBrush(LoadObject<UTexture2D>(nullptr, TEXT("/AudioWidgets/Resources/ScrubHandleDown_Clamped.ScrubHandleDown_Clamped")), FMargin(6.f / 13.f, 3.f / 12.f, 6.f / 13.f, 7.f / 12.f)));
 }
 
+void FAudioWidgetsStyle::SetupStyles()
+{
+	const FVector2f Icon16x16 = FVector2f(16.0f, 16.0f);
+	const FVector2f Icon20x20 = FVector2f(20.0f, 20.0f);
+
+	Set("AudioWidgetsStyle.FadeIn", new IMAGE_BRUSH_SVG(TEXT("Slate/Icons/FadeIn_20"), Icon20x20));
+	Set("AudioWidgetsStyle.FadeInLinear", new IMAGE_BRUSH_SVG(TEXT("Slate/Icons/FadeInLinear_16"), Icon16x16));
+	Set("AudioWidgetsStyle.FadeInLogarithmic", new IMAGE_BRUSH_SVG(TEXT("Slate/Icons/FadeInLogarithmic_16"), Icon16x16));
+	Set("AudioWidgetsStyle.FadeInExponential", new IMAGE_BRUSH_SVG(TEXT("Slate/Icons/FadeInExponential_16"), Icon16x16));
+	Set("AudioWidgetsStyle.FadeInSigmoid", new IMAGE_BRUSH_SVG(TEXT("Slate/Icons/FadeInSigmoid_16"), Icon16x16));
+	
+	Set("AudioWidgetsStyle.FadeOut", new IMAGE_BRUSH_SVG(TEXT("Slate/Icons/FadeOut_20"), Icon20x20));
+	Set("AudioWidgetsStyle.FadeOutLinear", new IMAGE_BRUSH_SVG(TEXT("Slate/Icons/FadeOutLinear_16"), Icon16x16));
+	Set("AudioWidgetsStyle.FadeOutLogarithmic", new IMAGE_BRUSH_SVG(TEXT("Slate/Icons/FadeOutLogarithmic_16"), Icon16x16));
+	Set("AudioWidgetsStyle.FadeOutExponential", new IMAGE_BRUSH_SVG(TEXT("Slate/Icons/FadeOutExponential_16"), Icon16x16));
+	Set("AudioWidgetsStyle.FadeOutSigmoid", new IMAGE_BRUSH_SVG(TEXT("Slate/Icons/FadeOutSigmoid_16"), Icon16x16));
+
+	Set("AudioWidgetsStyle.MarkerAdd", new IMAGE_BRUSH_SVG(TEXT("Slate/Icons/Marker_20"), Icon20x20));
+	Set("AudioWidgetsStyle.MarkerDelete", new IMAGE_BRUSH_SVG(TEXT("Slate/Icons/MarkerDelete_20"), Icon20x20));
+}
+
 FAudioWidgetsStyle::~FAudioWidgetsStyle()
 {
 	FSlateStyleRegistry::UnRegisterSlateStyle(*this);
diff --git a/Engine/Plugins/Runtime/AudioWidgets/Source/AudioWidgets/Public/AudioWidgetsStyle.h b/Engine/Plugins/Runtime/AudioWidgets/Source/AudioWidgets/Public/AudioWidgetsStyle.h
--- a/Engine/Plugins/Runtime/AudioWidgets/Source/AudioWidgets/Public/AudioWidgetsStyle.h
+++ b/Engine/Plugins/Runtime/AudioWidgets/Source/AudioWidgets/Public/AudioWidgetsStyle.h
@@ -21,6 +21,7 @@ public:
 
 private:
 	void SetResources();
+	void SetupStyles();
 
 	FAudioWidgetsStyle();
 	~FAudioWidgetsStyle();
diff --git a/Engine/Source/Editor/EditorStyle/Private/StarshipStyle.cpp b/Engine/Source/Editor/EditorStyle/Private/StarshipStyle.cpp
--- a/Engine/Source/Editor/EditorStyle/Private/StarshipStyle.cpp
+++ b/Engine/Source/Editor/EditorStyle/Private/StarshipStyle.cpp
@@ -2261,6 +2261,7 @@ void FStarshipEditorStyle::FStyle::SetupSequencerStyles()
 		Set("Sequencer.Transport.JumpToNextKey", new IMAGE_BRUSH_SVG("Sequencer/PlaybackControls/PlayControlsJumpToNextKey", Icon20x20));
 		Set("Sequencer.Transport.SetPlayStart", new IMAGE_BRUSH_SVG("Sequencer/PlaybackControls/PlayControlsSetPlaybackStart", Icon20x20));
 		Set("Sequencer.Transport.SetPlayEnd", new IMAGE_BRUSH_SVG("Sequencer/PlaybackControls/PlayControlsSetPlaybackEnd", Icon20x20));
+		Set("Sequencer.Transport.Looping", new IMAGE_BRUSH_SVG("Sequencer/PlaybackControls/PlayControlsLooping", Icon20x20));
 
 		Set("Sequencer.Transport.CloseButton", FButtonStyle()
 			.SetNormal(IMAGE_BRUSH("/Docking/CloseApp_Normal", Icon16x16))
