[Waveform Editor] - Default loop size when creating a loop in the processing panel is too small. Set default loop size to 10% of the total wave size to be consistent between all three ways of creating a loop

diff --git a/Engine/Plugins/Editor/WaveformEditor/Source/WaveformTransformations/Private/WaveformTransformationMarkers.cpp b/Engine/Plugins/Editor/WaveformEditor/Source/WaveformTransformations/Private/WaveformTransformationMarkers.cpp
--- a/Engine/Plugins/Editor/WaveformEditor/Source/WaveformTransformations/Private/WaveformTransformationMarkers.cpp
+++ b/Engine/Plugins/Editor/WaveformEditor/Source/WaveformTransformations/Private/WaveformTransformationMarkers.cpp
@@ -94,29 +94,37 @@ void UWaveformTransformationMarkers::PostEditChangeProperty(FPropertyChangedEven
 	Super::PostEditChangeProperty(PropertyChangedEvent);
 
 	const FName PropertyName = PropertyChangedEvent.GetPropertyName();
-	if (bIsPreviewingLoopRegion && PropertyName == GET_MEMBER_NAME_CHECKED(UWaveformTransformationMarkers, Markers))
+	if (PropertyName == GET_MEMBER_NAME_CHECKED(UWaveformTransformationMarkers, Markers))
 	{
 		check(Markers);
-		for (const FSoundWaveCuePoint& CuePoint : Markers->CuesAndLoops)
+		
+		for (FSoundWaveCuePoint& CuePoint : Markers->CuesAndLoops)
 		{
-			// Update LoopRegion Preview window if a loop is selected and the user edits the loop bounds using the properties window
-			if (Markers->SelectedCue == CuePoint.CuePointID)
+			if(CuePoint.IsLoopRegion())
 			{
-				check(CuePoint.IsLoopRegion());
-				check(SampleRate > 0.f);
-				check(EndLoopTime > 0)
+				// Update LoopRegion Preview window if a loop is selected and the user edits the loop bounds using the properties window
+				if (bIsPreviewingLoopRegion && Markers->SelectedCue == CuePoint.CuePointID)
+				{
+					check(SampleRate > 0.f);
+					check(EndLoopTime > 0.0);
 
-				const int64 StartLoopFramePos = StartLoopTime * SampleRate;
-				const int64 EndLoopFramePos = EndLoopTime * SampleRate;
+					const int64 StartLoopFramePos = StartLoopTime * SampleRate;
+					const int64 EndLoopFramePos = EndLoopTime * SampleRate;
 
-				if (CuePoint.FramePosition != StartLoopFramePos)
-				{
-					StartLoopTime = static_cast<float>(CuePoint.FramePosition) / SampleRate;
+					if (CuePoint.FramePosition != StartLoopFramePos)
+					{
+						StartLoopTime = static_cast<float>(CuePoint.FramePosition) / SampleRate;
+					}
+
+					if (CuePoint.FramePosition + CuePoint.FrameLength != EndLoopFramePos)
+					{
+						EndLoopTime = static_cast<float>(CuePoint.FramePosition + CuePoint.FrameLength) / SampleRate;
+					}
 				}
 
-				if (CuePoint.FramePosition + CuePoint.FrameLength != EndLoopFramePos)
+				if (CuePoint.FrameLength < Markers->MinLoopSize)
 				{
-					EndLoopTime = static_cast<float>(CuePoint.FramePosition + CuePoint.FrameLength) / SampleRate;
+					CuePoint.FrameLength = FMath::Max(Markers->MinLoopSize, static_cast<int64>(AvailableWaveformDuration * SampleRate * 0.1f));
 				}
 			}
 		}
@@ -133,12 +141,17 @@ void UWaveformTransformationMarkers::UpdateConfiguration(FWaveTransformUObjectCo
 {
 	check(Markers != nullptr);
 	Markers->InitMarkersIfNotSet(InOutConfiguration.WaveCues);
+	
+	const float InOutDuration = InOutConfiguration.EndTime - InOutConfiguration.StartTime;
 
 	// Assert that InOutConfiguration is initialized and valid
 	check(InOutConfiguration.SampleRate > 0.f);
 	check(InOutConfiguration.StartTime >= 0.f);
+	check(InOutDuration > 0.f);
+
 	StartFrameOffset = static_cast<int64>(InOutConfiguration.StartTime * InOutConfiguration.SampleRate);
 	SampleRate = InOutConfiguration.SampleRate;
+	AvailableWaveformDuration = InOutDuration;
 
 	if (!bCachedIsPreviewingLoopRegion)
 	{
@@ -221,14 +234,6 @@ void UWaveCueArray::PostEditChangeProperty(FPropertyChangedEvent& PropertyChange
 {
 	Super::PostEditChangeProperty(PropertyChangedEvent);
 
-	for (FSoundWaveCuePoint& CuePoint : CuesAndLoops)
-	{
-		if (CuePoint.IsLoopRegion() && CuePoint.FrameLength < MinLoopSize)
-		{
-			CuePoint.FrameLength = MinLoopSize;
-		}
-	}
-
 	const FName PropertyName = PropertyChangedEvent.GetPropertyName();
 	if (PropertyName == GET_MEMBER_NAME_CHECKED(UWaveCueArray, CuesAndLoops) && CuesAndLoops.Num() > 0)
 	{
