Waveform Editor: Fix for undo markers triggering ensure

diff --git a/Engine/Plugins/Editor/WaveformEditor/Source/WaveformTransformationsWidgets/Private/WaveformTransformationMarkersRenderer.cpp b/Engine/Plugins/Editor/WaveformEditor/Source/WaveformTransformationsWidgets/Private/WaveformTransformationMarkersRenderer.cpp
--- a/Engine/Plugins/Editor/WaveformEditor/Source/WaveformTransformationsWidgets/Private/WaveformTransformationMarkersRenderer.cpp
+++ b/Engine/Plugins/Editor/WaveformEditor/Source/WaveformTransformationsWidgets/Private/WaveformTransformationMarkersRenderer.cpp
@@ -806,7 +806,6 @@ void FWaveformTransformationMarkerRenderer::SetActiveLoopRegion(int32 CuePointID
 {
 	check(MarkersArray != nullptr);
 
-	MarkersHandle->NotifyPreChange();
 	MarkersArray->Modify();
 	MarkersArray->SelectedCue = CuePointID;
 
