Waveform Editor: UI fixes for markers with frame length

diff --git a/Engine/Plugins/Editor/WaveformEditor/Source/WaveformTransformationsWidgets/Private/WaveformTransformationMarkersRenderer.cpp b/Engine/Plugins/Editor/WaveformEditor/Source/WaveformTransformationsWidgets/Private/WaveformTransformationMarkersRenderer.cpp
--- a/Engine/Plugins/Editor/WaveformEditor/Source/WaveformTransformationsWidgets/Private/WaveformTransformationMarkersRenderer.cpp
+++ b/Engine/Plugins/Editor/WaveformEditor/Source/WaveformTransformationsWidgets/Private/WaveformTransformationMarkersRenderer.cpp
@@ -224,10 +224,12 @@ void FWaveformTransformationMarkerRenderer::SetPropertyHandles(const TArray<TSha
 	EndLoopTimeHandle = GetPropertyHandle(GET_MEMBER_NAME_CHECKED(UWaveformTransformationMarkers, EndLoopTime));
 	IsPreviewingLoopRegionHandle = GetPropertyHandle(GET_MEMBER_NAME_CHECKED(UWaveformTransformationMarkers, bIsPreviewingLoopRegion));
 
-	MarkersArray = CastChecked<UWaveCueArray>(GetPropertyValue<UObject*>(MarkersHandle));
+	MarkersArray = TStrongObjectPtr<UWaveCueArray>(CastChecked<UWaveCueArray>(GetPropertyValue<UObject*>(MarkersHandle)));
 
 	MarkersArray->ModifyMarkerLoop.BindLambda([this](ELoopModificationControls Modification) { ModifyMarkerLoopRegion(Modification); });
 	MarkersArray->CycleMarkerLoop.BindLambda([this](ELoopModificationControls Modification) { CycleMarkerLoopRegion(Modification); });
+
+	MarkersArray->CueChanged.BindLambda([this]() {SetActiveLoopRegion(MarkersArray->SelectedCue); });
 }
 
 void FWaveformTransformationMarkerRenderer::ModifyMarkerLoopRegion(ELoopModificationControls Modification)
@@ -764,7 +766,7 @@ bool FWaveformTransformationMarkerRenderer::SetMarkerInInteractionRange(const FV
 	// Evaluate loop area if no handles are selected
 	for (FSoundWaveCuePoint& CuePoint : MarkersArray->CuesAndLoops)
 	{
-		if (CuePoint.FrameLength > 0 && IsPositionInLoopArea(InLocalCursorPosition, WidgetGeometry, &CuePoint))
+		if (CuePoint.IsLoopRegion() && IsPositionInLoopArea(InLocalCursorPosition, WidgetGeometry, &CuePoint))
 		{
 			const int64 LoopRegionPosition = static_cast<int64>(InLocalCursorPosition.X / PixelsPerFrame - static_cast<double>(CuePoint.FramePosition));
 
@@ -789,7 +791,7 @@ bool FWaveformTransformationMarkerRenderer::SelectLoopArea(const FVector2D& InLo
 	for (FSoundWaveCuePoint& LoopRegion : LoopPoints)
 	{
 		// Check mouse position bounds
-		if (IsPositionInLoopArea(InLocalCursorPosition, WidgetGeometry, &LoopRegion))
+		if (LoopRegion.IsLoopRegion() && IsPositionInLoopArea(InLocalCursorPosition, WidgetGeometry, &LoopRegion))
 		{
 			SetActiveLoopRegion(LoopRegion.CuePointID);
 
@@ -811,7 +813,7 @@ void FWaveformTransformationMarkerRenderer::SetActiveLoopRegion(int32 CuePointID
 
 	check(TransformationWaveInfo.SampleRate > 0);
 	// Revert Loop preview handles
-	if (CuePointID == INDEX_NONE)
+	if (CuePointID == INDEX_NONE || GetSelectedMarker() == nullptr || !GetSelectedMarker()->IsLoopRegion())
 	{
 		SetIsPreviewingLoopRegion(0.0, -1.0, false);
 	}
diff --git a/Engine/Plugins/Editor/WaveformEditor/Source/WaveformTransformationsWidgets/Private/WaveformTransformationMarkersRenderer.h b/Engine/Plugins/Editor/WaveformEditor/Source/WaveformTransformationsWidgets/Private/WaveformTransformationMarkersRenderer.h
--- a/Engine/Plugins/Editor/WaveformEditor/Source/WaveformTransformationsWidgets/Private/WaveformTransformationMarkersRenderer.h
+++ b/Engine/Plugins/Editor/WaveformEditor/Source/WaveformTransformationsWidgets/Private/WaveformTransformationMarkersRenderer.h
@@ -83,7 +83,7 @@ private:
 	TArray<FSoundWaveCuePoint> CuePoints;
 	TArray<FSoundWaveCuePoint> LoopPoints;
 
-	TObjectPtr<UWaveCueArray> MarkersArray;
+	TStrongObjectPtr<UWaveCueArray> MarkersArray;
 
 	double PixelsPerFrame = 0.0;
 
