[Waveform Editor] - Highlighting and saving with a loop region selected disables the rest of the soundwave playing

diff --git a/Engine/Plugins/Editor/WaveformEditor/Source/WaveformEditor/Private/WaveformEditor.cpp b/Engine/Plugins/Editor/WaveformEditor/Source/WaveformEditor/Private/WaveformEditor.cpp
--- a/Engine/Plugins/Editor/WaveformEditor/Source/WaveformEditor/Private/WaveformEditor.cpp
+++ b/Engine/Plugins/Editor/WaveformEditor/Source/WaveformEditor/Private/WaveformEditor.cpp
@@ -132,20 +132,28 @@ FWaveformEditor::~FWaveformEditor()
 		FReimportManager::Instance()->OnPostReimport().RemoveAll(this);
 	}
 
-	if (TransformationChainConfig.bIsPreviewingLoopRegion)
-	{
-		// Note: if the user doesn't save the soundwave asset, bLooping will not be reset.
-		check(SoundWave);
-		SoundWave->Modify();
-		SoundWave->bLooping = TransformationChainConfig.bCachedSoundWaveLoopState;
-		SoundWave->PostEditChange();
-	}
-
 	TObjectPtr<UWaveformTransformationMarkers> MarkerTransformation = GetMarkerTransformation();
 
-	if (MarkerTransformation != nullptr && bCueChangeRegisteredByWaveformEditor)
+	if (MarkerTransformation != nullptr)
 	{
-		MarkerTransformation->Markers->CueChanged.Unbind();
+		if (bCueChangeRegisteredByWaveformEditor)
+		{
+			check(MarkerTransformation->Markers);
+			MarkerTransformation->Markers->CueChanged.Unbind();
+		}
+
+		if (TransformationChainConfig.bIsPreviewingLoopRegion)
+		{
+			// Note: if the user doesn't save the soundwave asset, bLooping will not be reset.
+			check(SoundWave);
+			SoundWave->Modify();
+			SoundWave->bLooping = TransformationChainConfig.bCachedSoundWaveLoopState;
+
+			MarkerTransformation->ResetLoopPreviewing();
+			SoundWave->UpdateTransformations();
+
+			SoundWave->PostEditChange();
+		}
 	}
 }
 
@@ -427,6 +435,7 @@ void FWaveformEditor::DeleteMarker()
 
 	TObjectPtr<UWaveformTransformationMarkers> MarkerTransformation = GetOrAddMarkerTransformation();
 	check(MarkerTransformation);
+	check(MarkerTransformation->Markers);
 
 	int32 CueToDelete = MarkerTransformation->Markers->SelectedCue;
 
@@ -442,8 +451,14 @@ void FWaveformEditor::DeleteMarker()
 	{
 		if (CueToDelete == MarkerTransformation->Markers->CuesAndLoops[i].CuePointID)
 		{
+			if (MarkerTransformation->Markers->CuesAndLoops[i].IsLoopRegion())
+			{
+				MarkerTransformation->ResetLoopPreviewing();
+			}
+
 			MarkerTransformation->Markers->CuesAndLoops.RemoveAt(i);
 			MarkerTransformation->Markers->SelectedCue = INDEX_NONE;
+
 			break;
 		}
 	}
diff --git a/Engine/Plugins/Editor/WaveformEditor/Source/WaveformTransformations/Private/WaveformTransformationMarkers.cpp b/Engine/Plugins/Editor/WaveformEditor/Source/WaveformTransformations/Private/WaveformTransformationMarkers.cpp
--- a/Engine/Plugins/Editor/WaveformEditor/Source/WaveformTransformations/Private/WaveformTransformationMarkers.cpp
+++ b/Engine/Plugins/Editor/WaveformEditor/Source/WaveformTransformations/Private/WaveformTransformationMarkers.cpp
@@ -97,9 +97,15 @@ void UWaveformTransformationMarkers::PostEditChangeProperty(FPropertyChangedEven
 	if (PropertyName == GET_MEMBER_NAME_CHECKED(UWaveformTransformationMarkers, Markers))
 	{
 		check(Markers);
+		bool bIsSelectedCuePresent = false;
 		
 		for (FSoundWaveCuePoint& CuePoint : Markers->CuesAndLoops)
 		{
+			if (!bIsSelectedCuePresent && Markers->SelectedCue == CuePoint.CuePointID)
+			{
+				bIsSelectedCuePresent = true;
+			}
+
 			if(CuePoint.IsLoopRegion())
 			{
 				// Update LoopRegion Preview window if a loop is selected and the user edits the loop bounds using the properties window
@@ -128,6 +134,16 @@ void UWaveformTransformationMarkers::PostEditChangeProperty(FPropertyChangedEven
 				}
 			}
 		}
+
+		if (!bIsSelectedCuePresent && Markers->SelectedCue != INDEX_NONE)
+		{
+			Markers->SelectedCue = INDEX_NONE;
+		}
+
+		if (Markers->SelectedCue == INDEX_NONE && (bIsPreviewingLoopRegion || StartLoopTime != 0.0 || EndLoopTime != -1.0))
+		{
+			ResetLoopPreviewing();
+		}
 	}
 }
 #endif // WITH_EDITOR
@@ -188,6 +204,13 @@ void UWaveformTransformationMarkers::CycleMarkerLoopRegion(ELoopModificationCont
 	Markers->CycleMarkerLoop.ExecuteIfBound(Modification);
 }
 
+void UWaveformTransformationMarkers::ResetLoopPreviewing()
+{
+	bIsPreviewingLoopRegion = false;
+	StartLoopTime = 0.0;
+	EndLoopTime = -1.0;
+}
+
 #if WITH_EDITOR
 void UWaveformTransformationMarkers::OverwriteSoundWaveData(USoundWave& InOutSoundWave)
 {
diff --git a/Engine/Plugins/Editor/WaveformEditor/Source/WaveformTransformations/Public/WaveformTransformationMarkers.h b/Engine/Plugins/Editor/WaveformEditor/Source/WaveformTransformations/Public/WaveformTransformationMarkers.h
--- a/Engine/Plugins/Editor/WaveformEditor/Source/WaveformTransformations/Public/WaveformTransformationMarkers.h
+++ b/Engine/Plugins/Editor/WaveformEditor/Source/WaveformTransformations/Public/WaveformTransformationMarkers.h
@@ -105,6 +105,8 @@ public:
 	void ModifyMarkerLoopRegion(ELoopModificationControls Modification) const;
 	void CycleMarkerLoopRegion(ELoopModificationControls Modification) const;
 
+	void ResetLoopPreviewing();
+
 #if WITH_EDITOR
 	void OverwriteSoundWaveData(USoundWave& InOutSoundWave) override;
 	void GetTransformationInfo(FWaveformTransformationInfo& InOutTransformationInfo) const override;
