Waveform Editor: Fix ensure on undoing loop deletion

diff --git a/Engine/Plugins/Editor/WaveformEditor/Source/WaveformTransformationsWidgets/Private/WaveformTransformationMarkersRenderer.cpp b/Engine/Plugins/Editor/WaveformEditor/Source/WaveformTransformationsWidgets/Private/WaveformTransformationMarkersRenderer.cpp
--- a/Engine/Plugins/Editor/WaveformEditor/Source/WaveformTransformationsWidgets/Private/WaveformTransformationMarkersRenderer.cpp
+++ b/Engine/Plugins/Editor/WaveformEditor/Source/WaveformTransformationsWidgets/Private/WaveformTransformationMarkersRenderer.cpp
@@ -229,7 +229,8 @@ void FWaveformTransformationMarkerRenderer::SetPropertyHandles(const TArray<TSha
 	MarkersArray->ModifyMarkerLoop.BindLambda([this](ELoopModificationControls Modification) { ModifyMarkerLoopRegion(Modification); });
 	MarkersArray->CycleMarkerLoop.BindLambda([this](ELoopModificationControls Modification) { CycleMarkerLoopRegion(Modification); });
 
-	MarkersArray->CueChanged.BindLambda([this]() {SetActiveLoopRegion(MarkersArray->SelectedCue); });
+	// Flag that the loop preview needs to refresh next Tick to avoid issues starting transactions during an Undo
+	MarkersArray->CueChanged.BindLambda([this]() { bLoopRegionPreviewNeedsRefresh = true; });
 }
 
 void FWaveformTransformationMarkerRenderer::ModifyMarkerLoopRegion(ELoopModificationControls Modification)
@@ -496,6 +497,13 @@ void FWaveformTransformationMarkerRenderer::Tick(const FGeometry& AllottedGeomet
 		}
 	}
 	CuePointMap.Empty();
+
+	if (bLoopRegionPreviewNeedsRefresh)
+	{
+		SetActiveLoopRegion(MarkersArray->SelectedCue);
+
+		bLoopRegionPreviewNeedsRefresh = false;
+	}
 }
 
 FReply FWaveformTransformationMarkerRenderer::OnMouseButtonDown(SWidget& OwnerWidget, const FGeometry& MyGeometry, const FPointerEvent& MouseEvent)
diff --git a/Engine/Plugins/Editor/WaveformEditor/Source/WaveformTransformationsWidgets/Private/WaveformTransformationMarkersRenderer.h b/Engine/Plugins/Editor/WaveformEditor/Source/WaveformTransformationsWidgets/Private/WaveformTransformationMarkersRenderer.h
--- a/Engine/Plugins/Editor/WaveformEditor/Source/WaveformTransformationsWidgets/Private/WaveformTransformationMarkersRenderer.h
+++ b/Engine/Plugins/Editor/WaveformEditor/Source/WaveformTransformationsWidgets/Private/WaveformTransformationMarkersRenderer.h
@@ -100,6 +100,8 @@ private:
 	FLinearColor LabelTextColor = WaveformTransformationWidgetSharedDefaults::DefaultLabelTextColor;
 	float LabelFontSize = WaveformTransformationWidgetSharedDefaults::DefaultLabelFontSize;
 
+	bool bLoopRegionPreviewNeedsRefresh = false;
+
 	TSharedPtr<IPropertyHandle> MarkersHandle;
 	TSharedPtr<IPropertyHandle> StartLoopTimeHandle;
 	TSharedPtr<IPropertyHandle> EndLoopTimeHandle;
