[Waveform Editor] - Editor crashes when opening a level while the wave editor is open and has a transform on it
crash on UE startup with soundwaves set to auto open.
This is a temp workaround for UE - 265361

diff --git a/Engine/Plugins/Editor/WaveformEditor/Source/WaveformEditor/Private/WaveformEditor.cpp b/Engine/Plugins/Editor/WaveformEditor/Source/WaveformEditor/Private/WaveformEditor.cpp
--- a/Engine/Plugins/Editor/WaveformEditor/Source/WaveformEditor/Private/WaveformEditor.cpp
+++ b/Engine/Plugins/Editor/WaveformEditor/Source/WaveformEditor/Private/WaveformEditor.cpp
@@ -7,6 +7,7 @@
 #include "AssetDefinitionRegistry.h"
 #include "Components/AudioComponent.h"
 #include "EditorReimportHandler.h"
+#include "Framework/Notifications/NotificationManager.h"
 #include "Misc/TransactionObjectEvent.h"
 #include "PropertyEditorModule.h"
 #include "Sound/SoundWave.h"
@@ -21,6 +22,7 @@
 #include "WaveformEditorCommands.h"
 #include "WaveformEditorDetailsCustomization.h"
 #include "WaveformEditorLog.h"
+#include "WaveformEditorModule.h"
 #include "WaveformEditorSequenceDataProvider.h"
 #include "WaveformEditorStyle.h"
 #include "WaveformEditorToolMenuContext.h"
@@ -28,6 +30,7 @@
 #include "WaveformEditorWaveWriter.h"
 #include "WaveformTransformationMarkers.h"
 #include "Widgets/Docking/SDockTab.h"
+#include "Widgets/Notifications/SNotificationList.h"
 
 #define LOCTEXT_NAMESPACE "WaveformEditor"
 
@@ -38,11 +41,53 @@ const FName FWaveformEditor::WaveformDisplayTabId("WaveformEditor_Display");
 const FName FWaveformEditor::EditorName("Waveform Editor");
 const FName FWaveformEditor::ToolkitFName("WaveformEditor");
 
+/* 
+ * This is a temp workaround for UE - 265361
+ * Prevent the waveform editor from opening and causing a crash if the user has the editor auto open a USoundWave asset on editor startup
+ */
+TSharedPtr<FWaveformEditor> FWaveformEditor::CreateIfAble()
+{
+	const bool bIsMapLoaded = FModuleManager::GetModuleChecked<FWaveformEditorModule>("WaveformEditor").GetIsLevelLoaded();
+
+	if (!bIsMapLoaded)
+	{
+		return nullptr;
+	}
+
+	return MakeShared<FWaveformEditor>();
+}
 
 bool FWaveformEditor::Init(const EToolkitMode::Type Mode, const TSharedPtr<IToolkitHost>& InitToolkitHost, USoundWave* SoundWaveToEdit)
 {
 	checkf(SoundWaveToEdit, TEXT("Tried to open a Soundwave Editor from a null soundwave"));
 
+	const bool bIsMapLoaded = FModuleManager::GetModuleChecked<FWaveformEditorModule>("WaveformEditor").GetIsLevelLoaded();
+
+	const TSharedRef<FTabManager::FLayout>  StandaloneDefaultLayout = SetupStandaloneLayout();
+
+	const bool bCreateDefaultStandaloneMenu = true;
+	const bool bCreateDefaultToolbar = true;
+	const bool bToolbarFocusable = false;
+	const bool bUseSmallIcons = true;
+
+	// Prevent the waveform editor from opening and causing a crash if the user has the editor auto open a USoundWave asset on editor startup
+	if (!bIsMapLoaded)
+	{
+		// Open the asset editor so it can be closed in the destructor
+		FAssetEditorToolkit::InitAssetEditor(
+			Mode,
+			InitToolkitHost,
+			AppIdentifier,
+			StandaloneDefaultLayout,
+			bCreateDefaultStandaloneMenu,
+			bCreateDefaultToolbar,
+			SoundWaveToEdit,
+			bToolbarFocusable,
+			bUseSmallIcons);
+
+		return false;
+	}
+
 	SoundWave = SoundWaveToEdit;
 
 	// Initialize TransformationChainConfig from persisted transformations or if there are none falling back to default transformations.
@@ -70,17 +115,10 @@ bool FWaveformEditor::Init(const EToolkitMode::Type Mode, const TSharedPtr<ITool
 	bIsInitialized &= RegisterToolbar();
 	bIsInitialized &= BindCommands();
 
-	const TSharedRef<FTabManager::FLayout>  StandaloneDefaultLayout = SetupStandaloneLayout();
-
 	GEditor->RegisterForUndo(this);
 
 	if (bIsInitialized)
 	{
-		const bool bCreateDefaultStandaloneMenu = true;
-		const bool bCreateDefaultToolbar = true;
-		const bool bToolbarFocusable = false;
-		const bool bUseSmallIcons = true;
-
 		FAssetEditorToolkit::InitAssetEditor(
 			Mode,
 			InitToolkitHost,
@@ -103,15 +141,61 @@ bool FWaveformEditor::Init(const EToolkitMode::Type Mode, const TSharedPtr<ITool
 			check(TransportCoordinator != nullptr);
 			TransportCoordinator->SetProgressRatio(0.0f);
 
-			OnMapOpenedHandle = FEditorDelegates::OnMapOpened.AddLambda([this](const FString&, bool)
+			OnMapOpenedHandle = FEditorDelegates::OnMapOpened.AddLambda([this, SoundWaveToEdit](const FString&, bool)
 				{
-					if (TransformationsDetails.IsValid())
+					CloseWindow(EAssetEditorCloseReason::EditorRefreshRequested);
+
+					FNotificationInfo Info(FText::FromString(TEXT("Do you want to reopen the sound wave for editing?")));
+					Info.bFireAndForget = true;
+					Info.ExpireDuration = 15.0f;
+
+					Info.ButtonDetails.Add(FNotificationButtonInfo(
+						FText::FromString(TEXT("Yes")),
+						FText::FromString(TEXT("Open the SoundWave asset")),
+						FSimpleDelegate::CreateLambda([this, SoundWaveToEdit]() 
+							{
+								bool bSuccess = false;
+								if (GEditor && SoundWaveToEdit)
+								{
+									TArray<UObject*> AssetsToOpen = { SoundWaveToEdit };
+									if (UAssetEditorSubsystem* AssetEditorSubsystem = GEditor->GetEditorSubsystem<UAssetEditorSubsystem>())
+									{
+										bSuccess = AssetEditorSubsystem->OpenEditorForAssets(AssetsToOpen);
+									}
+								}
+								if (ReopenNotificationItem.IsValid())
+								{
+									if (bSuccess)
+									{
+										ReopenNotificationItem->SetCompletionState(SNotificationItem::CS_Success);
+									}
+									else
+									{
+										ReopenNotificationItem->SetCompletionState(SNotificationItem::CS_Fail);
+									}
+									ReopenNotificationItem->Fadeout();
+								}
+							})
+					));
+
+					Info.ButtonDetails.Add(FNotificationButtonInfo(
+						FText::FromString(TEXT("No")),
+						FText::FromString(TEXT("Dismiss this message")),
+						FSimpleDelegate::CreateLambda([this]()
+							{
+								if (ReopenNotificationItem.IsValid())
+								{
+									ReopenNotificationItem->SetCompletionState(SNotificationItem::CS_None);
+									ReopenNotificationItem->Fadeout();
+								}
+							})
+					));
+
+					ReopenNotificationItem = FSlateNotificationManager::Get().AddNotification(Info);
+					
+					if (ReopenNotificationItem.IsValid())
 					{
-						TransformationsDetails->ForceRefresh();
-					}
-					if (PropertiesDetails.IsValid())
-					{
-						PropertiesDetails->ForceRefresh();
+						ReopenNotificationItem->SetCompletionState(SNotificationItem::CS_Pending);
 					}
 				});
 	}
@@ -268,6 +352,11 @@ TObjectPtr<UWaveformTransformationMarkers> FWaveformEditor::GetOrAddMarkerTransf
 
 TObjectPtr<UWaveformTransformationMarkers> FWaveformEditor::GetMarkerTransformation()
 {
+	if (SoundWave == nullptr)
+	{
+		return nullptr;
+	}
+
 	TObjectPtr<UWaveformTransformationMarkers> MarkersTransformation;
 
 	for (TObjectPtr<UWaveformTransformationBase> Transformation : SoundWave->Transformations)
@@ -650,6 +739,11 @@ void FWaveformEditor::ExecuteOverwriteTransformations()
 
 void FWaveformEditor::RegisterTabSpawners(const TSharedRef<FTabManager>& InTabManager)
 {
+	if (TransformationsDetails == nullptr)
+	{
+		return;
+	}
+
 	WorkspaceMenuCategory = InTabManager->AddLocalWorkspaceMenuCategory(LOCTEXT("WorkspaceMenu_WaveformEditor", "Sound Wave Editor"));
 
 	FAssetEditorToolkit::RegisterTabSpawners(InTabManager);
diff --git a/Engine/Plugins/Editor/WaveformEditor/Source/WaveformEditor/Private/WaveformEditorInstantiator.cpp b/Engine/Plugins/Editor/WaveformEditor/Source/WaveformEditor/Private/WaveformEditorInstantiator.cpp
--- a/Engine/Plugins/Editor/WaveformEditor/Source/WaveformEditor/Private/WaveformEditorInstantiator.cpp
+++ b/Engine/Plugins/Editor/WaveformEditor/Source/WaveformEditor/Private/WaveformEditorInstantiator.cpp
@@ -16,6 +16,7 @@
 #include "ToolMenuSection.h"
 #include "WaveformEditor.h"
 #include "WaveformEditorLog.h"
+#include "WaveformEditorModule.h"
 
 #define LOCTEXT_NAMESPACE "WaveformEditorInstantiator"
 
@@ -156,12 +157,17 @@ void FWaveformEditorInstantiator::CreateWaveformEditor(TArray<USoundWave*> Sound
 	{
 		if (CanSoundWaveBeOpenedInEditor(SoundWavePtr))
 		{
-			TSharedRef<FWaveformEditor> WaveformEditor = MakeShared<FWaveformEditor>();
+			// CreateIfAble is a temp workaround for UE - 265361
+			TSharedPtr<FWaveformEditor> WaveformEditor = FWaveformEditor::CreateIfAble();
 
-			if (!WaveformEditor->Init(EToolkitMode::Standalone, nullptr, SoundWavePtr))
+			if (WaveformEditor == nullptr || !WaveformEditor->Init(EToolkitMode::Standalone, nullptr, SoundWavePtr))
 			{
-				UE_LOG(LogWaveformEditor, Warning, TEXT("Could not open waveform editor for soundwave %s, initialization failed"), *(SoundWavePtr->GetName()))
-				WaveformEditor->CloseWindow(EAssetEditorCloseReason::AssetUnloadingOrInvalid);
+				UE_LOG(LogWaveformEditor, Warning, TEXT("Could not open waveform editor for soundwave %s, initialization failed"), *(SoundWavePtr->GetName()));
+				
+				if (WaveformEditor)
+				{
+					WaveformEditor->CloseWindow(EAssetEditorCloseReason::AssetUnloadingOrInvalid);
+				}
 			}
 		}
 	}
diff --git a/Engine/Plugins/Editor/WaveformEditor/Source/WaveformEditor/Private/WaveformEditorModule.cpp b/Engine/Plugins/Editor/WaveformEditor/Source/WaveformEditor/Private/WaveformEditorModule.cpp
--- a/Engine/Plugins/Editor/WaveformEditor/Source/WaveformEditor/Private/WaveformEditorModule.cpp
+++ b/Engine/Plugins/Editor/WaveformEditor/Source/WaveformEditor/Private/WaveformEditorModule.cpp
@@ -2,6 +2,7 @@
 
 #include "WaveformEditorModule.h"
 
+#include "Editor.h"
 #include "Modules/ModuleManager.h"
 #include "PropertyEditorModule.h"
 #include "TransformedWaveformViewFactory.h"
@@ -20,11 +21,30 @@ void FWaveformEditorModule::StartupModule()
 	WaveformEditorInstantiator = MakeShared<FWaveformEditorInstantiator>();
 	RegisterContentBrowserExtensions(WaveformEditorInstantiator.Get());
 	WaveformEditorInstantiator->RegisterAsSoundwaveEditor();
+
+	FEditorDelegates::OnEditorBoot.AddLambda([this](double)
+		{
+			bIsLevelLoaded = false;
+		});
+
+	FWorldDelegates::OnWorldTickEnd.AddLambda([this](UWorld*, ELevelTick LevelTick, float)
+		{
+			// Once Tick Type is LEVELTICK_ViewportsOnly we know the level and viewports are fully loaded, we can now open the waveform editor
+			if (LevelTick == ELevelTick::LEVELTICK_ViewportsOnly)
+			{
+				bIsLevelLoaded = true;
+				FWorldDelegates::OnWorldTickEnd.RemoveAll(this);
+			}
+		});
+
 }
 
 void FWaveformEditorModule::ShutdownModule()
 {
 	FWaveformEditorCommands::Unregister();
+
+	FEditorDelegates::OnEditorBoot.RemoveAll(this);
+	FWorldDelegates::OnWorldTickEnd.RemoveAll(this);
 }
 
 void FWaveformEditorModule::RegisterContentBrowserExtensions(IWaveformEditorInstantiator* Instantiator)
diff --git a/Engine/Plugins/Editor/WaveformEditor/Source/WaveformEditor/Public/WaveformEditor.h b/Engine/Plugins/Editor/WaveformEditor/Source/WaveformEditor/Public/WaveformEditor.h
--- a/Engine/Plugins/Editor/WaveformEditor/Source/WaveformEditor/Public/WaveformEditor.h
+++ b/Engine/Plugins/Editor/WaveformEditor/Source/WaveformEditor/Public/WaveformEditor.h
@@ -25,7 +25,7 @@ class WAVEFORMEDITOR_API FWaveformEditor
 	, public FNotifyHook
 {
 public:
-
+	static TSharedPtr<FWaveformEditor> CreateIfAble();
 	bool Init(const EToolkitMode::Type Mode, const TSharedPtr<IToolkitHost>& InitToolkitHost, USoundWave* SoundWaveToEdit);
 	virtual ~FWaveformEditor();
 
@@ -188,5 +188,6 @@ private:
 	FText GetExportButtonToolTip() const;
 
 	FDelegateHandle OnMapOpenedHandle;
+	TSharedPtr<SNotificationItem> ReopenNotificationItem;
 };
 
diff --git a/Engine/Plugins/Editor/WaveformEditor/Source/WaveformEditor/Public/WaveformEditorModule.h b/Engine/Plugins/Editor/WaveformEditor/Source/WaveformEditor/Public/WaveformEditorModule.h
--- a/Engine/Plugins/Editor/WaveformEditor/Source/WaveformEditor/Public/WaveformEditorModule.h
+++ b/Engine/Plugins/Editor/WaveformEditor/Source/WaveformEditor/Public/WaveformEditorModule.h
@@ -13,9 +13,13 @@ public:
 	virtual void StartupModule() override;
 	virtual void ShutdownModule() override;
 
+	bool GetIsLevelLoaded() const { return bIsLevelLoaded; }
+
 private:
 	virtual void RegisterContentBrowserExtensions(IWaveformEditorInstantiator* Instantiator) override;
 
 	TSharedPtr<IWaveformEditorInstantiator> WaveformEditorInstantiator = nullptr;
+
+	bool bIsLevelLoaded = false;
 };
 
