LowLevelTests Audio for waveform editor failing

diff --git a/Engine/Plugins/Editor/WaveformEditor/Source/WaveformEditor/Private/WaveformEditor.cpp b/Engine/Plugins/Editor/WaveformEditor/Source/WaveformEditor/Private/WaveformEditor.cpp
--- a/Engine/Plugins/Editor/WaveformEditor/Source/WaveformEditor/Private/WaveformEditor.cpp
+++ b/Engine/Plugins/Editor/WaveformEditor/Source/WaveformEditor/Private/WaveformEditor.cpp
@@ -57,7 +57,7 @@ TSharedPtr<FWaveformEditor> FWaveformEditor::CreateIfAble()
 	return MakeShared<FWaveformEditor>();
 }
 
-bool FWaveformEditor::Init(const EToolkitMode::Type Mode, const TSharedPtr<IToolkitHost>& InitToolkitHost, USoundWave* SoundWaveToEdit)
+bool FWaveformEditor::Init(const EToolkitMode::Type Mode, const TSharedPtr<IToolkitHost>& InitToolkitHost, USoundWave* SoundWaveToEdit, bool bIsAutomationTest)
 {
 	checkf(SoundWaveToEdit, TEXT("Tried to open a Soundwave Editor from a null soundwave"));
 
@@ -71,7 +71,7 @@ bool FWaveformEditor::Init(const EToolkitMode::Type Mode, const TSharedPtr<ITool
 	const bool bUseSmallIcons = true;
 
 	// Prevent the waveform editor from opening and causing a crash if the user has the editor auto open a USoundWave asset on editor startup
-	if (!bIsMapLoaded)
+	if (!bIsMapLoaded && !bIsAutomationTest)
 	{
 		// Open the asset editor so it can be closed in the destructor
 		FAssetEditorToolkit::InitAssetEditor(
diff --git a/Engine/Plugins/Editor/WaveformEditor/Source/WaveformEditor/Public/WaveformEditor.h b/Engine/Plugins/Editor/WaveformEditor/Source/WaveformEditor/Public/WaveformEditor.h
--- a/Engine/Plugins/Editor/WaveformEditor/Source/WaveformEditor/Public/WaveformEditor.h
+++ b/Engine/Plugins/Editor/WaveformEditor/Source/WaveformEditor/Public/WaveformEditor.h
@@ -26,7 +26,7 @@ class WAVEFORMEDITOR_API FWaveformEditor
 {
 public:
 	static TSharedPtr<FWaveformEditor> CreateIfAble();
-	bool Init(const EToolkitMode::Type Mode, const TSharedPtr<IToolkitHost>& InitToolkitHost, USoundWave* SoundWaveToEdit);
+	bool Init(const EToolkitMode::Type Mode, const TSharedPtr<IToolkitHost>& InitToolkitHost, USoundWave* SoundWaveToEdit, bool bIsAutomationTest = false);
 	virtual ~FWaveformEditor();
 
 	/** FAssetEditorToolkit interface */
