[Waveform Editor] Removing waveform editor as default soundwave asset editor and removing crash workaround.

diff --git a/Engine/Plugins/Editor/WaveformEditor/Source/WaveformEditor/Private/WaveformEditor.cpp b/Engine/Plugins/Editor/WaveformEditor/Source/WaveformEditor/Private/WaveformEditor.cpp
--- a/Engine/Plugins/Editor/WaveformEditor/Source/WaveformEditor/Private/WaveformEditor.cpp
+++ b/Engine/Plugins/Editor/WaveformEditor/Source/WaveformEditor/Private/WaveformEditor.cpp
@@ -41,28 +41,10 @@ const FName FWaveformEditor::WaveformDisplayTabId("WaveformEditor_Display");
 const FName FWaveformEditor::EditorName("Waveform Editor");
 const FName FWaveformEditor::ToolkitFName("WaveformEditor");
 
-/* 
- * This is a temp workaround for UE - 265361
- * Prevent the waveform editor from opening and causing a crash if the user has the editor auto open a USoundWave asset on editor startup
- */
-TSharedPtr<FWaveformEditor> FWaveformEditor::CreateIfAble()
-{
-	const bool bIsMapLoaded = FModuleManager::GetModuleChecked<FWaveformEditorModule>("WaveformEditor").GetIsLevelLoaded();
-
-	if (!bIsMapLoaded)
-	{
-		return nullptr;
-	}
-
-	return MakeShared<FWaveformEditor>();
-}
-
-bool FWaveformEditor::Init(const EToolkitMode::Type Mode, const TSharedPtr<IToolkitHost>& InitToolkitHost, USoundWave* SoundWaveToEdit, bool bIsAutomationTest)
+bool FWaveformEditor::Init(const EToolkitMode::Type Mode, const TSharedPtr<IToolkitHost>& InitToolkitHost, USoundWave* SoundWaveToEdit)
 {
 	checkf(SoundWaveToEdit, TEXT("Tried to open a Soundwave Editor from a null soundwave"));
 
-	const bool bIsMapLoaded = FModuleManager::GetModuleChecked<FWaveformEditorModule>("WaveformEditor").GetIsLevelLoaded();
-
 	const TSharedRef<FTabManager::FLayout>  StandaloneDefaultLayout = SetupStandaloneLayout();
 
 	const bool bCreateDefaultStandaloneMenu = true;
@@ -70,24 +52,6 @@ bool FWaveformEditor::Init(const EToolkitMode::Type Mode, const TSharedPtr<ITool
 	const bool bToolbarFocusable = false;
 	const bool bUseSmallIcons = true;
 
-	// Prevent the waveform editor from opening and causing a crash if the user has the editor auto open a USoundWave asset on editor startup
-	if (!bIsMapLoaded && !bIsAutomationTest)
-	{
-		// Open the asset editor so it can be closed in the destructor
-		FAssetEditorToolkit::InitAssetEditor(
-			Mode,
-			InitToolkitHost,
-			AppIdentifier,
-			StandaloneDefaultLayout,
-			bCreateDefaultStandaloneMenu,
-			bCreateDefaultToolbar,
-			SoundWaveToEdit,
-			bToolbarFocusable,
-			bUseSmallIcons);
-
-		return false;
-	}
-
 	SoundWave = SoundWaveToEdit;
 
 	// Initialize TransformationChainConfig from persisted transformations or if there are none falling back to default transformations.
diff --git a/Engine/Plugins/Editor/WaveformEditor/Source/WaveformEditor/Private/WaveformEditorInstantiator.cpp b/Engine/Plugins/Editor/WaveformEditor/Source/WaveformEditor/Private/WaveformEditorInstantiator.cpp
--- a/Engine/Plugins/Editor/WaveformEditor/Source/WaveformEditor/Private/WaveformEditorInstantiator.cpp
+++ b/Engine/Plugins/Editor/WaveformEditor/Source/WaveformEditor/Private/WaveformEditorInstantiator.cpp
@@ -29,6 +29,22 @@ void FWaveformEditorInstantiator::ExtendContentBrowserSelectionMenu()
 {
 	UToolMenu* Menu = UToolMenus::Get()->ExtendMenu("ContentBrowser.AssetContextMenu.SoundWave");
 	FToolMenuSection& Section = Menu->FindOrAddSection("GetAssetActions");
+
+	Section.AddDynamicEntry("SoundWaveEditing_CreateWaveformEditor", FNewToolMenuSectionDelegate::CreateLambda([this](FToolMenuSection& InSection)
+		{
+			if (const UContentBrowserAssetContextMenuContext* Context = InSection.FindContext<UContentBrowserAssetContextMenuContext>())
+			{
+				if (Algo::AnyOf(Context->SelectedAssets, WaveEditorInstantiatorPrivate::FilterUnwantedAssets))
+				{
+					const TAttribute<FText> Label = LOCTEXT("SoundWave_WaveformEditor", "Edit Waveform");
+					const TAttribute<FText> ToolTip = LOCTEXT("SoundWave_WaveformEditor_Tooltip", "Open waveform editor");
+					const FSlateIcon Icon = FSlateIcon();
+					const FToolMenuExecuteAction UIAction = FToolMenuExecuteAction::CreateSP(this, &FWaveformEditorInstantiator::ExecuteCreateWaveformEditor);
+
+					InSection.AddMenuEntry("SoundWave_CreateWaveformEditor", Label, ToolTip, Icon, UIAction);
+				}
+			}
+		}));
 	
 	FToolMenuSection& ImportSection = Menu->FindOrAddSection("ImportedAssetActions");
 
@@ -60,6 +76,15 @@ void FWaveformEditorInstantiator::RegisterAsSoundwaveEditor()
 	}
 }
 
+void FWaveformEditorInstantiator::ExecuteCreateWaveformEditor(const FToolMenuContext& MenuContext)
+{
+	if (const UContentBrowserAssetContextMenuContext* Context = UContentBrowserAssetContextMenuContext::FindContextWithAssets(MenuContext))
+	{
+		TArray<USoundWave*> SoundWavesToEdit = Context->LoadSelectedObjectsIf<USoundWave>(WaveEditorInstantiatorPrivate::FilterUnwantedAssets);
+		CreateWaveformEditor(SoundWavesToEdit);
+	}
+}
+
 void FWaveformEditorInstantiator::ExecuteReImportAndOverwriteTransformations(const FToolMenuContext& MenuContext)
 {
 	if (const UContentBrowserAssetContextMenuContext* Context = UContentBrowserAssetContextMenuContext::FindContextWithAssets(MenuContext))
@@ -157,8 +182,7 @@ void FWaveformEditorInstantiator::CreateWaveformEditor(TArray<USoundWave*> Sound
 	{
 		if (CanSoundWaveBeOpenedInEditor(SoundWavePtr))
 		{
-			// CreateIfAble is a temp workaround for UE - 265361
-			TSharedPtr<FWaveformEditor> WaveformEditor = FWaveformEditor::CreateIfAble();
+			TSharedPtr<FWaveformEditor> WaveformEditor = MakeShared<FWaveformEditor>();
 
 			if (WaveformEditor == nullptr || !WaveformEditor->Init(EToolkitMode::Standalone, nullptr, SoundWavePtr))
 			{
diff --git a/Engine/Plugins/Editor/WaveformEditor/Source/WaveformEditor/Private/WaveformEditorInstantiator.h b/Engine/Plugins/Editor/WaveformEditor/Source/WaveformEditor/Private/WaveformEditorInstantiator.h
--- a/Engine/Plugins/Editor/WaveformEditor/Source/WaveformEditor/Private/WaveformEditorInstantiator.h
+++ b/Engine/Plugins/Editor/WaveformEditor/Source/WaveformEditor/Private/WaveformEditorInstantiator.h
@@ -19,6 +19,7 @@ public:
 
 private:
 	virtual void CreateWaveformEditor(TArray<USoundWave*> SoundWavesToEdit) override;
+	void ExecuteCreateWaveformEditor(const FToolMenuContext& MenuContext);
 	void ExecuteReImportAndOverwriteTransformations(const FToolMenuContext& MenuContext);
 	bool CanSoundWaveBeOpenedInEditor(const USoundWave* SoundWaveToEdit);
 	void DisplayErrorDialog(const FText& ErrorMessage) const;
diff --git a/Engine/Plugins/Editor/WaveformEditor/Source/WaveformEditor/Private/WaveformEditorModule.cpp b/Engine/Plugins/Editor/WaveformEditor/Source/WaveformEditor/Private/WaveformEditorModule.cpp
--- a/Engine/Plugins/Editor/WaveformEditor/Source/WaveformEditor/Private/WaveformEditorModule.cpp
+++ b/Engine/Plugins/Editor/WaveformEditor/Source/WaveformEditor/Private/WaveformEditorModule.cpp
@@ -20,31 +20,13 @@ void FWaveformEditorModule::StartupModule()
 
 	WaveformEditorInstantiator = MakeShared<FWaveformEditorInstantiator>();
 	RegisterContentBrowserExtensions(WaveformEditorInstantiator.Get());
-	WaveformEditorInstantiator->RegisterAsSoundwaveEditor();
-
-	FEditorDelegates::OnEditorBoot.AddLambda([this](double)
-		{
-			bIsLevelLoaded = false;
-		});
-
-	FWorldDelegates::OnWorldTickEnd.AddLambda([this](UWorld*, ELevelTick LevelTick, float)
-		{
-			// Once Tick Type is LEVELTICK_ViewportsOnly we know the level and viewports are fully loaded, we can now open the waveform editor
-			if (LevelTick == ELevelTick::LEVELTICK_ViewportsOnly)
-			{
-				bIsLevelLoaded = true;
-				FWorldDelegates::OnWorldTickEnd.RemoveAll(this);
-			}
-		});
 
+	// TODO: Re register WaveformEditor as default soundwave editor with: WaveformEditorInstantiator->RegisterAsSoundwaveEditor();
 }
 
 void FWaveformEditorModule::ShutdownModule()
 {
 	FWaveformEditorCommands::Unregister();
-
-	FEditorDelegates::OnEditorBoot.RemoveAll(this);
-	FWorldDelegates::OnWorldTickEnd.RemoveAll(this);
 }
 
 void FWaveformEditorModule::RegisterContentBrowserExtensions(IWaveformEditorInstantiator* Instantiator)
diff --git a/Engine/Plugins/Editor/WaveformEditor/Source/WaveformEditor/Public/WaveformEditor.h b/Engine/Plugins/Editor/WaveformEditor/Source/WaveformEditor/Public/WaveformEditor.h
--- a/Engine/Plugins/Editor/WaveformEditor/Source/WaveformEditor/Public/WaveformEditor.h
+++ b/Engine/Plugins/Editor/WaveformEditor/Source/WaveformEditor/Public/WaveformEditor.h
@@ -25,8 +25,7 @@ class WAVEFORMEDITOR_API FWaveformEditor
 	, public FNotifyHook
 {
 public:
-	static TSharedPtr<FWaveformEditor> CreateIfAble();
-	bool Init(const EToolkitMode::Type Mode, const TSharedPtr<IToolkitHost>& InitToolkitHost, USoundWave* SoundWaveToEdit, bool bIsAutomationTest = false);
+	bool Init(const EToolkitMode::Type Mode, const TSharedPtr<IToolkitHost>& InitToolkitHost, USoundWave* SoundWaveToEdit);
 	virtual ~FWaveformEditor();
 
 	/** FAssetEditorToolkit interface */
diff --git a/Engine/Plugins/Editor/WaveformEditor/Source/WaveformEditor/Public/WaveformEditorModule.h b/Engine/Plugins/Editor/WaveformEditor/Source/WaveformEditor/Public/WaveformEditorModule.h
--- a/Engine/Plugins/Editor/WaveformEditor/Source/WaveformEditor/Public/WaveformEditorModule.h
+++ b/Engine/Plugins/Editor/WaveformEditor/Source/WaveformEditor/Public/WaveformEditorModule.h
@@ -13,13 +13,9 @@ public:
 	virtual void StartupModule() override;
 	virtual void ShutdownModule() override;
 
-	bool GetIsLevelLoaded() const { return bIsLevelLoaded; }
-
 private:
 	virtual void RegisterContentBrowserExtensions(IWaveformEditorInstantiator* Instantiator) override;
 
 	TSharedPtr<IWaveformEditorInstantiator> WaveformEditorInstantiator = nullptr;
-
-	bool bIsLevelLoaded = false;
 };
 
