Waveform Editor: Fix frame length resetting to default instead of min when sliding

diff --git a/Engine/Plugins/Editor/WaveformEditor/Source/WaveformTransformations/Private/WaveformTransformationMarkers.cpp b/Engine/Plugins/Editor/WaveformEditor/Source/WaveformTransformations/Private/WaveformTransformationMarkers.cpp
--- a/Engine/Plugins/Editor/WaveformEditor/Source/WaveformTransformations/Private/WaveformTransformationMarkers.cpp
+++ b/Engine/Plugins/Editor/WaveformEditor/Source/WaveformTransformations/Private/WaveformTransformationMarkers.cpp
@@ -130,7 +130,7 @@ void UWaveformTransformationMarkers::PostEditChangeProperty(FPropertyChangedEven
 
 				if (CuePoint.FrameLength < Markers->MinLoopSize)
 				{
-					CuePoint.FrameLength = FMath::Max(Markers->MinLoopSize, static_cast<int64>(AvailableWaveformDuration * SampleRate * 0.1f));
+					CuePoint.FrameLength = FMath::Min(Markers->MinLoopSize, static_cast<int64>(AvailableWaveformDuration * SampleRate * 0.1f));
 				}
 			}
 		}
