[Waveform editor] Fade duration not clamped consistently

diff --git a/Engine/Plugins/Editor/WaveformEditor/Source/WaveformTransformations/Private/WaveformTransformationTrimFade.cpp b/Engine/Plugins/Editor/WaveformEditor/Source/WaveformTransformations/Private/WaveformTransformationTrimFade.cpp
--- a/Engine/Plugins/Editor/WaveformEditor/Source/WaveformTransformations/Private/WaveformTransformationTrimFade.cpp
+++ b/Engine/Plugins/Editor/WaveformEditor/Source/WaveformTransformations/Private/WaveformTransformationTrimFade.cpp
@@ -140,6 +140,15 @@ void FWaveTransformationTrimFade::ProcessAudio(Audio::FWaveformTransformationWav
 	}
 }
 
+void UWaveformTransformationTrimFade::PostEditChangeProperty(FPropertyChangedEvent& PropertyChangedEvent)
+{
+	Super::PostEditChangeProperty(PropertyChangedEvent);
+
+	const float MaxDuration = EndTime - StartTime;
+	EndFadeTime = FMath::Min(EndFadeTime, MaxDuration);
+	StartFadeTime = FMath::Min(StartFadeTime, MaxDuration);
+}
+
 Audio::FTransformationPtr UWaveformTransformationTrimFade::CreateTransformation() const
 {
 	return MakeUnique<FWaveTransformationTrimFade>(StartTime, EndTime, StartFadeTime, StartFadeCurve, StartSCurveSharpness, EndFadeTime, EndFadeCurve, EndSCurveSharpness);
diff --git a/Engine/Plugins/Editor/WaveformEditor/Source/WaveformTransformations/Public/WaveformTransformationTrimFade.h b/Engine/Plugins/Editor/WaveformEditor/Source/WaveformTransformations/Public/WaveformTransformationTrimFade.h
--- a/Engine/Plugins/Editor/WaveformEditor/Source/WaveformTransformations/Public/WaveformTransformationTrimFade.h
+++ b/Engine/Plugins/Editor/WaveformEditor/Source/WaveformTransformations/Public/WaveformTransformationTrimFade.h
@@ -45,13 +45,18 @@ class UWaveformTransformationTrimFade : public UWaveformTransformationBase
 	GENERATED_BODY()
 public:
 
+#if WITH_EDITOR
+	UE_API virtual void PostEditChangeProperty(FPropertyChangedEvent& PropertyChangedEvent) override;
+#endif // WITH_EDITOR
+
 	UPROPERTY(EditAnywhere, Category = "Trim", meta=(ClampMin = 0.0))
 	double StartTime = 0.0;
 
 	UPROPERTY(EditAnywhere, Category = "Trim")
 	double EndTime = -1.0;
 
-	UPROPERTY(EditAnywhere, Category = "Fade", meta=(ClampMin = 0.0, DisplayName = "Fade-In Duration"))
+	// Relative to StartTime
+	UPROPERTY(EditAnywhere, Category = "Fade", meta = (ClampMin = 0.0, DisplayName = "Fade-In Duration"))
 	float StartFadeTime = 0.f;
 
 	UPROPERTY(EditAnywhere, Category = "Fade", meta=(ClampMin = -0.1, ClampMax = 10.0, DisplayName = "Fade-In Curve"))
@@ -60,7 +65,8 @@ public:
 	UPROPERTY(EditAnywhere, Category = "Fade", meta = (ClampMin = -1.0, ClampMax = 1.0, DisplayName = "S-Curve Sharpness", EditCondition = "StartFadeCurve < 0", EditConditionHides))
 	float StartSCurveSharpness = 0.1f;
 
-	UPROPERTY(EditAnywhere, Category = "Fade", meta=(ClampMin = 0.0, DisplayName = "Fade-Out Duration"))
+	// Relative to EndTime. Fade out starts at EndTime - EndFadeTime and ends at EndTime
+	UPROPERTY(EditAnywhere, Category = "Fade", meta = (ClampMin = 0.0, DisplayName = "Fade-Out Duration"))
 	float EndFadeTime = 0.f;
 
 	UPROPERTY(EditAnywhere, Category = "Fade", meta=(ClampMin = -0.1, ClampMax = 10.0, DisplayName = "Fade-Out Curve"))
