Unreal Editor Crashes after opening an exported Waveform Soundwave from content browser

diff --git a/Engine/Plugins/Editor/WaveformEditor/Source/WaveformEditor/Private/WaveformEditor.cpp b/Engine/Plugins/Editor/WaveformEditor/Source/WaveformEditor/Private/WaveformEditor.cpp
--- a/Engine/Plugins/Editor/WaveformEditor/Source/WaveformEditor/Private/WaveformEditor.cpp
+++ b/Engine/Plugins/Editor/WaveformEditor/Source/WaveformEditor/Private/WaveformEditor.cpp
@@ -1118,9 +1118,11 @@ bool FWaveformEditor::BindCommands()
 		FExecuteAction::CreateSP(ZoomManager.ToSharedRef(), &FWaveformEditorZoomController::ZoomOut),
 		FCanExecuteAction::CreateSP(ZoomManager.ToSharedRef(), &FWaveformEditorZoomController::CanZoomOut));
 
+	constexpr bool bPromptUserToSave = true;
+	
 	ToolkitCommands->MapAction(
 		Commands.ExportWaveform,
-		FExecuteAction::CreateSP(this, &FWaveformEditor::ExportWaveform),
+		FExecuteAction::CreateSP(this, &FWaveformEditor::ExportWaveform, bPromptUserToSave),
 		FCanExecuteAction::CreateSP(WaveWriter.ToSharedRef(), &FWaveformEditorWaveWriter::CanCreateSoundWaveAsset));
 
 	ToolkitCommands->MapAction(
@@ -1486,6 +1488,12 @@ void FWaveformEditor::InitToolMenuContext(FToolMenuContext& MenuContext)
 	MenuContext.AddObject(Context);
 }
 
+void FWaveformEditor::ExportWaveform(bool bPromptUserToSave)
+{
+	check(WaveWriter);
+	WaveWriter->ExportTransformedWaveform(bPromptUserToSave);
+}
+
 bool FWaveformEditor::CreateDetailsViews()
 {
 	if (!ensure(SoundWave)) 
@@ -1756,12 +1764,6 @@ bool FWaveformEditor::CreateWaveWriter()
 	return WaveWriter != nullptr;
 }
 
-void FWaveformEditor::ExportWaveform()
-{
-	check(WaveWriter);
-	WaveWriter->ExportTransformedWaveform();
-}
-
 const UWaveformEditorTransformationsSettings* FWaveformEditor::GetWaveformEditorTransformationsSettings() const
 {
 	const UWaveformEditorTransformationsSettings* WaveformEditorTransformationsSettings = GetDefault<UWaveformEditorTransformationsSettings>();
diff --git a/Engine/Plugins/Editor/WaveformEditor/Source/WaveformEditor/Private/WaveformEditorWaveWriter.cpp b/Engine/Plugins/Editor/WaveformEditor/Source/WaveformEditor/Private/WaveformEditorWaveWriter.cpp
--- a/Engine/Plugins/Editor/WaveformEditor/Source/WaveformEditor/Private/WaveformEditorWaveWriter.cpp
+++ b/Engine/Plugins/Editor/WaveformEditor/Source/WaveformEditor/Private/WaveformEditorWaveWriter.cpp
@@ -21,7 +21,7 @@ bool FWaveformEditorWaveWriter::CanCreateSoundWaveAsset() const
 	return SourceSoundWave != nullptr && WaveWriter != nullptr && WaveWriter->IsDone();
 }
 
-void FWaveformEditorWaveWriter::ExportTransformedWaveform()
+void FWaveformEditorWaveWriter::ExportTransformedWaveform(bool bPromptUserToSave)
 {
 	check(SourceSoundWave);
 	check(WaveWriter);
@@ -41,27 +41,25 @@ void FWaveformEditorWaveWriter::ExportTransformedWaveform()
 
 	Audio::TSampleBuffer<> BufferToWrite = GenerateSampleBuffer();
 
-	TFunction<void(const USoundWave*)> OnSoundWaveWritten = [AssetName, AssetPath](const USoundWave* ResultingWave) {
+	TFunction<void(const USoundWave*)> OnSoundWaveWritten = [bPromptUserToSave, AssetName, AssetPath](const USoundWave* ResultingWave) {
 		UE_LOG(LogWaveformEditor, Log, TEXT("Finished Exporting edited soundwave %s/%s"), *AssetPath, *AssetName);
 		TArray<UPackage*> PackagesToSave;
 
 		if(ResultingWave->GetPackage())
 		{
 			PackagesToSave.Add(ResultingWave->GetPackage());
-			FEditorFileUtils::EPromptReturnCode Result = FEditorFileUtils::PromptForCheckoutAndSave(PackagesToSave, false /*bCheckDirty*/, true /*bPromptToSave*/);
+
+			FEditorFileUtils::EPromptReturnCode Result = FEditorFileUtils::PromptForCheckoutAndSave(PackagesToSave, false /*bCheckDirty*/, bPromptUserToSave);
 
 			// If the user cancels the checkout and save, delete the new generated soundwave
-			if (Result == FEditorFileUtils::PR_Cancelled)
+			if (bPromptUserToSave && Result == FEditorFileUtils::PR_Cancelled)
 			{
 				UEditorAssetLibrary::DeleteAsset(ResultingWave->GetPathName());
 			}
 		}
 	};
 	
-	//Duplicate the SourceSoundWave so we do not overwrite the original file with the transformations on export
-	TObjectPtr<USoundWave> ExportSoundWave = DuplicateObject<USoundWave>(SourceSoundWave, SourceSoundWave->GetOuter());
-
-	if (!WaveWriter->BeginWriteToSoundWave(AssetName, BufferToWrite, AssetPath, OnSoundWaveWritten, ExportSoundWave))
+	if (!WaveWriter->BeginWriteToSoundWave(AssetName, BufferToWrite, AssetPath, OnSoundWaveWritten, SourceSoundWave))
 	{
 		UE_LOG(LogWaveformEditor, Log, TEXT("Exporting edited soundwave to %s/%s failed"), *AssetPath, *AssetName);
 	}
diff --git a/Engine/Plugins/Editor/WaveformEditor/Source/WaveformEditor/Private/WaveformEditorWaveWriter.h b/Engine/Plugins/Editor/WaveformEditor/Source/WaveformEditor/Private/WaveformEditorWaveWriter.h
--- a/Engine/Plugins/Editor/WaveformEditor/Source/WaveformEditor/Private/WaveformEditorWaveWriter.h
+++ b/Engine/Plugins/Editor/WaveformEditor/Source/WaveformEditor/Private/WaveformEditorWaveWriter.h
@@ -22,7 +22,7 @@ public:
 	explicit FWaveformEditorWaveWriter(USoundWave* InSoundWave);
 
 	bool CanCreateSoundWaveAsset() const;
-	void ExportTransformedWaveform();
+	void ExportTransformedWaveform(bool bPromptUserToSave);
 
 	WaveformEditorWaveWriter::EChannelFormat GetExportChannelsFormat() const;
 	void SetExportChannelsFormat(const WaveformEditorWaveWriter::EChannelFormat InTargetChannelFormat);
diff --git a/Engine/Plugins/Editor/WaveformEditor/Source/WaveformEditor/Public/WaveformEditor.h b/Engine/Plugins/Editor/WaveformEditor/Source/WaveformEditor/Public/WaveformEditor.h
--- a/Engine/Plugins/Editor/WaveformEditor/Source/WaveformEditor/Public/WaveformEditor.h
+++ b/Engine/Plugins/Editor/WaveformEditor/Source/WaveformEditor/Public/WaveformEditor.h
@@ -52,6 +52,8 @@ public:
 	UE_API bool MatchesContext(const FTransactionContext& InContext, const TArray<TPair<UObject*, FTransactionObjectEvent>>& TransactionObjectContexts) const override;
 	UE_API virtual void InitToolMenuContext(FToolMenuContext& MenuContext) override;
 
+	UE_API void ExportWaveform(bool bPromptUserToSave);
+
 private:
 	UE_API bool InitializeAudioComponent();
 	UE_API bool CreateTransportController();
@@ -99,7 +101,6 @@ private:
 	UE_API virtual FString GetReferencerName() const override;
 	UE_API bool CanPressPlayButton() const;
 	UE_API bool CreateWaveWriter();
-	UE_API void ExportWaveform();
 	UE_API const UWaveformEditorTransformationsSettings* GetWaveformEditorTransformationsSettings() const;
 	UE_API void AddDefaultTransformations();
 
diff --git a/Engine/Source/Editor/AudioEditor/Private/Factories/SoundFactory.cpp b/Engine/Source/Editor/AudioEditor/Private/Factories/SoundFactory.cpp
--- a/Engine/Source/Editor/AudioEditor/Private/Factories/SoundFactory.cpp
+++ b/Engine/Source/Editor/AudioEditor/Private/Factories/SoundFactory.cpp
@@ -437,12 +437,7 @@ UObject* USoundFactory::CreateObject
 			Sound->CompressionQuality = TemplateSoundWave->CompressionQuality;
 			Sound->SoundAssetCompressionType = TemplateSoundWave->SoundAssetCompressionType;
 
-			// we do not want to inherit these values from the template, as the data may be incorrect
-			// rather we re-parse them from the incoming file.
-			Sound->NumChannels = 0;
-			Sound->ChannelOffsets.Reset();
-			Sound->ChannelSizes.Reset();
-			Sound->Transformations.Empty();
+			Sound->ResetInvalidTemplateArguments();
 		}
 
 		int32 SizeOfSample = (*WaveInfo.pBitsPerSample) / 8;
diff --git a/Engine/Source/Runtime/Engine/Classes/Sound/SoundWave.h b/Engine/Source/Runtime/Engine/Classes/Sound/SoundWave.h
--- a/Engine/Source/Runtime/Engine/Classes/Sound/SoundWave.h
+++ b/Engine/Source/Runtime/Engine/Classes/Sound/SoundWave.h
@@ -502,6 +502,13 @@ public:
 	/** Returns the Runtime format of the wave */
 	ENGINE_API FName GetRuntimeFormat() const;
 
+	/**
+	* For use when creating a new soundwave with another soundwave used as a template
+	* We do not want to inherit these values from the template, as the data may be incorrect
+	* rather we re-parse them from the incoming file.
+	*/
+	ENGINE_API void ResetInvalidTemplateArguments();
+
 private:
 	// cached proxy
 	FSoundWaveProxyPtr Proxy{ nullptr };
@@ -615,8 +622,8 @@ public:
 #if WITH_EDITOR
 	ENGINE_API const FWaveTransformUObjectConfiguration& GetTransformationChainConfig() const;
 	ENGINE_API const FWaveTransformUObjectConfiguration& UpdateTransformations();
-	ENGINE_API TArray<uint8> CreateNewCueData(const char ChunkName[4]);
-	ENGINE_API void OverwriteSoundWaveDataWithTransformationData();
+	// Pass in a Transformations array to overwrite this soundwave with another soundwaves transformation data
+	ENGINE_API void OverwriteSoundWaveDataWithTransformationData(const TArrayView<TObjectPtr<UWaveformTransformationBase>> InTransformations = {});
 	ENGINE_API FWaveformTransformationInfo GetWaveformTransformationInfo() const;
 
 #endif
diff --git a/Engine/Source/Runtime/Engine/Private/SampleBufferIO.cpp b/Engine/Source/Runtime/Engine/Private/SampleBufferIO.cpp
--- a/Engine/Source/Runtime/Engine/Private/SampleBufferIO.cpp
+++ b/Engine/Source/Runtime/Engine/Private/SampleBufferIO.cpp
@@ -198,18 +198,12 @@ namespace Audio
 		// Create a new USoundWave.
 		CurrentSoundWave = NewObject<USoundWave>(CurrentPackage, *FileName, RF_Public | RF_Standalone | RF_Transactional, InSoundWaveTransformationTarget);
 		CurrentSoundWave->AddToRoot();
-
-#if WITH_EDITORONLY_DATA
-		// Clears transformations in the exported soundwave as they are now baked into the sound wave cue points
-		CurrentSoundWave->Transformations.Empty();
-#endif
-
+		
 #if WITH_EDITOR
-		if (InSoundWaveTransformationTarget != nullptr)
-		{
-			InSoundWaveTransformationTarget->OverwriteSoundWaveDataWithTransformationData();
-			CurrentSoundWave->SetSoundWaveCuePoints(InSoundWaveTransformationTarget->GetSoundWaveCuePoints());
-		}
+		CurrentSoundWave->ResetInvalidTemplateArguments();
+
+		// CurrentSoundwave::Transformations must be empty because all Transformation data will be written directly to the soundwave, not overlayed as a transformation
+		check(CurrentSoundWave->Transformations.IsEmpty());
 #endif //WITH_EDITOR
 
 		CurrentBuffer = InSampleBuffer;
@@ -220,12 +214,19 @@ namespace Audio
 			CurrentBuffer.MixBufferToChannels(2);
 		}
 	
-		CurrentOperation.Reset(new FAsyncSoundWavePCMWriterTask(this, ESoundWavePCMWriteTaskType::GenerateAndWriteSoundWave, [this, OnSuccess](const USoundWave* InSoundWave)
+		CurrentOperation.Reset(new FAsyncSoundWavePCMWriterTask(this, ESoundWavePCMWriteTaskType::GenerateAndWriteSoundWave, [this, OnSuccess, InSoundWaveTransformationTarget](const USoundWave* InSoundWave)
 		{
 			// Invalidate compressed data on new soundwave so that it gets refreshed or the soundwave is not seekable
 			// and causes this error: FMixerBuffer::CreateStreamingBuffer failed to StreamCompressedInfo on SoundWave
 			CurrentSoundWave->InvalidateCompressedData(true); // calls CurrentSoundWave->FreeResources(), note the game thread is the audio thread in the editor
 
+#if WITH_EDITOR
+			if (InSoundWaveTransformationTarget && InSoundWaveTransformationTarget->Transformations.Num() > 0)
+			{
+				CurrentSoundWave->OverwriteSoundWaveDataWithTransformationData(InSoundWaveTransformationTarget->Transformations);
+			}
+#endif //WITH_EDITOR
+
 			OnSuccess(InSoundWave);
 		}));
 		CurrentOperation->StartBackgroundTask();
diff --git a/Engine/Source/Runtime/Engine/Private/SoundWave.cpp b/Engine/Source/Runtime/Engine/Private/SoundWave.cpp
--- a/Engine/Source/Runtime/Engine/Private/SoundWave.cpp
+++ b/Engine/Source/Runtime/Engine/Private/SoundWave.cpp
@@ -1549,6 +1549,18 @@ FName USoundWave::GetRuntimeFormat() const
 	return SoundWaveDataPtr->GetRuntimeFormat();
 }
 
+void USoundWave::ResetInvalidTemplateArguments()
+{
+	NumChannels = 0;
+
+#if WITH_EDITOR
+	TotalSamples = 0;
+	ChannelOffsets.Reset();
+	ChannelSizes.Reset();
+	Transformations.Empty();
+#endif //WITH_EDITOR
+}
+
 float USoundWave::GetSubtitlePriority() const
 {
 	return SubtitlePriority;
@@ -4847,9 +4859,13 @@ const FWaveTransformUObjectConfiguration& USoundWave::UpdateTransformations()
 	return TransformationChainConfig;
 }
 
-void USoundWave::OverwriteSoundWaveDataWithTransformationData()
+void USoundWave::OverwriteSoundWaveDataWithTransformationData(const TArrayView<TObjectPtr<UWaveformTransformationBase>> InTransformations)
 {
-	for (TObjectPtr<UWaveformTransformationBase>& Transformation : Transformations)
+	// If InTransformations array is empty, use the current soundwaves transformations
+	TArrayView<TObjectPtr<UWaveformTransformationBase>> TransformationsView = InTransformations.IsEmpty() ? 
+		TArrayView<TObjectPtr<UWaveformTransformationBase>>(Transformations) : InTransformations;
+
+	for (const TObjectPtr<UWaveformTransformationBase>& Transformation : TransformationsView)
 	{
 		if (Transformation)
 		{
@@ -4897,6 +4913,17 @@ void USoundWave::OverwriteSoundWaveDataWithTransformationData()
 	TArray<uint8> InWaveFileData;
 	InWaveFileData.Append((uint8*)FutureBuffer.Get().GetData(), FutureBuffer.Get().GetSize());
 
+	// CreateCueAndSampleChunks errors out if there are more than 2 channels
+	if (NumChannels > 2)
+	{
+		if (!WaveCues.IsEmpty() && !WaveSampleLoops.IsEmpty())
+		{
+			UE_LOG(LogAudio, Warning, TEXT("Files with a channel count > 2 are not currently supported so new cue points or loop regions were not written to the RawData in file: %s"), *GetPathName());
+		}
+
+		return;
+	}
+
 	Audio::SoundFileUtils::CreateCueAndSampleChunks(InWaveFileData, WaveFileData, WaveCues, WaveSampleLoops);
 
 	// Update the RawData with the new payload
