[Waveform editor] Marker + fade context menus can be cut off

diff --git a/Engine/Plugins/Editor/WaveformEditor/Source/WaveformTransformationsWidgets/Private/WaveformTransformationMarkersRenderer.cpp b/Engine/Plugins/Editor/WaveformEditor/Source/WaveformTransformationsWidgets/Private/WaveformTransformationMarkersRenderer.cpp
--- a/Engine/Plugins/Editor/WaveformEditor/Source/WaveformTransformationsWidgets/Private/WaveformTransformationMarkersRenderer.cpp
+++ b/Engine/Plugins/Editor/WaveformEditor/Source/WaveformTransformationsWidgets/Private/WaveformTransformationMarkersRenderer.cpp
@@ -1003,6 +1003,7 @@ void FWaveformTransformationMarkerRenderer::ShowConvertToLoopAtMarker(const FGeo
 {
 	check(CueToModify);
 	check(MarkersArray != nullptr);
+	const FVector2D LocalWindowMaxPosition = MyGeometry.GetAbsolutePosition() + MyGeometry.GetAbsoluteSize();
 	const FVector2D LocalCursorPosition = MouseEvent.GetScreenSpacePosition();
 
 	TSharedRef<SVerticalBox> MenuContent = SNew(SVerticalBox);
@@ -1045,11 +1046,29 @@ void FWaveformTransformationMarkerRenderer::ShowConvertToLoopAtMarker(const FGeo
 		[
 			MenuContent
 		];
-
+	
 	CreateLoopMenuWindow = MenuWindow;
 
 	FSlateApplication::Get().AddWindow(MenuWindow);
 
+	FTSTicker::GetCoreTicker().AddTicker(FTickerDelegate::CreateLambda([this, LocalCursorPosition, LocalWindowMaxPosition](float DeltaTime)
+		{
+			if (TSharedPtr<SWindow> LockedLoopMenuWindow = CreateLoopMenuWindow.Pin())
+			{
+				const FVector2D WindowSize = LockedLoopMenuWindow->GetSizeInScreen();
+
+				if (WindowSize.X > 0.0 && LocalCursorPosition.X + WindowSize.X > LocalWindowMaxPosition.X)
+				{
+					check(WindowSize.Y > 0.0);
+
+					const double XPosition = FMath::Max(0.0, LocalCursorPosition.X - WindowSize.X);
+					LockedLoopMenuWindow->MoveWindowTo(FVector2D(XPosition, LocalCursorPosition.Y));
+				}
+			}
+
+			return false; // return false to only execute once
+		}));
+
 	FSlateApplication::Get().OnFocusChanging().Remove(PopupHandle);
 	FSlateApplication::Get().OnApplicationActivationStateChanged().Remove(ApplicationActivationStateHandle);
 
