[Waveform Editor] remove redundant dropdown arrow from TransformationMarker

diff --git a/Engine/Plugins/Editor/WaveformEditor/Source/WaveformTransformations/Private/WaveformTransformationMarkersObjectCustomization.cpp b/Engine/Plugins/Editor/WaveformEditor/Source/WaveformTransformations/Private/WaveformTransformationMarkersObjectCustomization.cpp
--- a/Engine/Plugins/Editor/WaveformEditor/Source/WaveformTransformations/Private/WaveformTransformationMarkersObjectCustomization.cpp
+++ b/Engine/Plugins/Editor/WaveformEditor/Source/WaveformTransformations/Private/WaveformTransformationMarkersObjectCustomization.cpp
@@ -16,8 +16,10 @@ TSharedRef<IPropertyTypeCustomization> FWaveformTransformationMarkersObjectCusto
 void FWaveformTransformationMarkersObjectCustomization::CustomizeHeader(TSharedRef<IPropertyHandle> PropertyHandle,
 	FDetailWidgetRow& HeaderRow, IPropertyTypeCustomizationUtils& CustomizationUtils)
 {
+	constexpr bool bDisplayDefaultPropertyButtons = false;
+
 	HeaderRow.NameContent()[PropertyHandle->CreatePropertyNameWidget()]
-	.ValueContent()[PropertyHandle->CreatePropertyValueWidget()];
+	.ValueContent()[PropertyHandle->CreatePropertyValueWidget(bDisplayDefaultPropertyButtons)];
 }
 
 void FWaveformTransformationMarkersObjectCustomization::CustomizeChildren(TSharedRef<IPropertyHandle> PropertyHandle,
