Waveform Editor: Improve UI handle grabbing

diff --git a/Engine/Plugins/Editor/WaveformEditor/Source/WaveformEditorWidgets/Private/SWaveformTransformationsOverlay.cpp b/Engine/Plugins/Editor/WaveformEditor/Source/WaveformEditorWidgets/Private/SWaveformTransformationsOverlay.cpp
--- a/Engine/Plugins/Editor/WaveformEditor/Source/WaveformEditorWidgets/Private/SWaveformTransformationsOverlay.cpp
+++ b/Engine/Plugins/Editor/WaveformEditor/Source/WaveformEditorWidgets/Private/SWaveformTransformationsOverlay.cpp
@@ -93,25 +93,25 @@ void SWaveformTransformationsOverlay::UpdateLayerConstraints()
 
 FReply SWaveformTransformationsOverlay::OnMouseButtonDown(const FGeometry& MyGeometry, const FPointerEvent& MouseEvent)
 {
-	return AudioWidgetsUtils::RouteMouseInput(&SWidget::OnMouseButtonDown, MouseEvent, MakeArrayView(TransformationLayers.GetData(), TransformationLayers.Num()), true);
+	return AudioWidgetsUtils::RouteMouseInput(&SWidget::OnMouseButtonDown, MouseEvent, MakeArrayView(TransformationLayers.GetData(), TransformationLayers.Num()), false);
 }
 
 FReply SWaveformTransformationsOverlay::OnMouseButtonUp(const FGeometry& MyGeometry, const FPointerEvent& MouseEvent)
 {
-	return AudioWidgetsUtils::RouteMouseInput(&SWidget::OnMouseButtonUp, MouseEvent, MakeArrayView(TransformationLayers.GetData(), TransformationLayers.Num()), true);
+	return AudioWidgetsUtils::RouteMouseInput(&SWidget::OnMouseButtonUp, MouseEvent, MakeArrayView(TransformationLayers.GetData(), TransformationLayers.Num()), false);
 }
 
 FReply SWaveformTransformationsOverlay::OnMouseMove(const FGeometry& MyGeometry, const FPointerEvent& MouseEvent)
 {
-	return AudioWidgetsUtils::RouteMouseInput(&SWidget::OnMouseMove, MouseEvent, MakeArrayView(TransformationLayers.GetData(), TransformationLayers.Num()), true);
+	return AudioWidgetsUtils::RouteMouseInput(&SWidget::OnMouseMove, MouseEvent, MakeArrayView(TransformationLayers.GetData(), TransformationLayers.Num()), false);
 }
 
 FReply SWaveformTransformationsOverlay::OnMouseWheel(const FGeometry& MyGeometry, const FPointerEvent& MouseEvent)
 {
-	return AudioWidgetsUtils::RouteMouseInput(&SWidget::OnMouseWheel, MouseEvent, MakeArrayView(TransformationLayers.GetData(), TransformationLayers.Num()), true);
+	return AudioWidgetsUtils::RouteMouseInput(&SWidget::OnMouseWheel, MouseEvent, MakeArrayView(TransformationLayers.GetData(), TransformationLayers.Num()), false);
 }
 
 FCursorReply SWaveformTransformationsOverlay::OnCursorQuery(const FGeometry& MyGeometry, const FPointerEvent& CursorEvent) const
 {
-	return AudioWidgetsUtils::RouteCursorQuery(CursorEvent, MakeArrayView(TransformationLayers.GetData(), TransformationLayers.Num()), true);
+	return AudioWidgetsUtils::RouteCursorQuery(CursorEvent, MakeArrayView(TransformationLayers.GetData(), TransformationLayers.Num()), false);
 }
