Cue Point Origin: improve defaults

diff --git a/Engine/Plugins/Editor/WaveformEditor/Source/WaveformEditor/Private/WaveformEditor.cpp b/Engine/Plugins/Editor/WaveformEditor/Source/WaveformEditor/Private/WaveformEditor.cpp
--- a/Engine/Plugins/Editor/WaveformEditor/Source/WaveformEditor/Private/WaveformEditor.cpp
+++ b/Engine/Plugins/Editor/WaveformEditor/Source/WaveformEditor/Private/WaveformEditor.cpp
@@ -105,6 +105,11 @@ bool FWaveformEditor::Init(const EToolkitMode::Type Mode, const TSharedPtr<ITool
 			check(TransportCoordinator != nullptr);
 			TransportCoordinator->SetProgressRatio(0.0f);
 
+			if (GetMarkerTransformation())
+			{
+				bMarkerTransformationIsActive = true;
+			}
+
 			OnMapOpenedHandle = FEditorDelegates::OnMapOpened.AddLambda([this, SoundWaveToEdit](const FString&, bool)
 				{
 					CloseWindow(EAssetEditorCloseReason::EditorRefreshRequested);
@@ -175,6 +180,11 @@ FWaveformEditor::~FWaveformEditor()
 		OnMapOpenedHandle.Reset();
 	}
 
+	if (ChangeCuePointOriginNotificationItem.IsValid())
+	{
+		ChangeCuePointOriginNotificationItem->Fadeout();
+	}
+
 	if (FReimportManager::Instance())
 	{
 		FReimportManager::Instance()->OnPostReimport().RemoveAll(this);
@@ -309,12 +319,13 @@ TObjectPtr<UWaveformTransformationMarkers> FWaveformEditor::GetOrAddMarkerTransf
 	if (MarkersTransformation == nullptr)
 	{
 		MarkersTransformation = AddMarkerTransformation();
+		PromptToChangeCuePointOrigin();
 	}
 
 	return MarkersTransformation;
 }
 
-TObjectPtr<UWaveformTransformationMarkers> FWaveformEditor::GetMarkerTransformation()
+TObjectPtr<UWaveformTransformationMarkers> FWaveformEditor::GetMarkerTransformation() const
 {
 	if (SoundWave == nullptr)
 	{
@@ -593,6 +604,98 @@ void FWaveformEditor::CycleMarkerLoopRegion(ELoopModificationControls Modificati
 	MarkersTransformation->CycleMarkerLoopRegion(Modification);
 }
 
+void FWaveformEditor::SetCuePointOrigin(TObjectPtr<UWaveformTransformationMarkers> MarkerTransformation)
+{
+	if (GEditor)
+	{
+		if (MarkerTransformation && SoundWave->GetCuePointOrigin() == ESoundWaveCuePointOrigin::WaveFile)
+		{
+			GEditor->BeginTransaction(TEXT("PropertyEditor"), LOCTEXT("SettingCuePointOrigin", "SetCuePointOrigin"), SoundWave);
+			SoundWave->Modify();
+			SoundWave->SetCuePointOrigin(ESoundWaveCuePointOrigin::MarkerTransformation);
+			SoundWave->PostEditChange();
+			GEditor->EndTransaction();
+		}
+		else if (SoundWave->GetCuePointOrigin() != ESoundWaveCuePointOrigin::WaveFile)
+		{
+			GEditor->BeginTransaction(TEXT("PropertyEditor"), LOCTEXT("SettingCuePointOrigin", "SetCuePointOrigin"), SoundWave);
+			SoundWave->Modify();
+			SoundWave->SetCuePointOrigin(ESoundWaveCuePointOrigin::WaveFile);
+			SoundWave->PostEditChange();
+			GEditor->EndTransaction();
+		}
+	}
+}
+
+void FWaveformEditor::PromptToChangeCuePointOrigin()
+{
+	if (bMarkerTransformationIsActive)
+	{
+		return; // Marker transformation already exists, do not prompt user
+	}
+
+	TObjectPtr<UWaveformTransformationMarkers> MarkerTransformation = GetMarkerTransformation(); // if marker transformation is valid, it is newly created
+
+	// Checking if the notification item is valid will only allow the user to be prompted once while the waveform editor is open (resets when closed)
+	if (MarkerTransformation && SoundWave->GetCuePointOrigin() == ESoundWaveCuePointOrigin::WaveFile && !ChangeCuePointOriginNotificationItem.IsValid())
+	{
+		bMarkerTransformationIsActive = true;
+
+		FNotificationInfo Info(FText::FromString(TEXT("Do you want to set your soundwave to use cues and loops from your WaveformTransformations?")));
+		Info.bFireAndForget = true;
+		Info.ExpireDuration = 15.0f;
+
+		Info.ButtonDetails.Add(FNotificationButtonInfo(
+			FText::FromString(TEXT("Yes")),
+			FText::FromString(TEXT("Set Cue Point Origin to Waveform Transformation Markers")),
+			FSimpleDelegate::CreateLambda([this]()
+				{
+					bool bSuccess = false;
+					if (GEditor && SoundWave)
+					{
+						if (TObjectPtr<UWaveformTransformationMarkers> MarkerTransformation = GetMarkerTransformation())
+						{
+							SetCuePointOrigin(MarkerTransformation);
+							bSuccess = true;
+						}
+					}
+					if (ChangeCuePointOriginNotificationItem.IsValid())
+					{
+						if (bSuccess)
+						{
+							ChangeCuePointOriginNotificationItem->SetCompletionState(SNotificationItem::CS_Success);
+						}
+						else
+						{
+							ChangeCuePointOriginNotificationItem->SetCompletionState(SNotificationItem::CS_Fail);
+						}
+						ChangeCuePointOriginNotificationItem->Fadeout();
+					}
+				})
+		));
+
+		Info.ButtonDetails.Add(FNotificationButtonInfo(
+			FText::FromString(TEXT("No")),
+			FText::FromString(TEXT("Dismiss this message")),
+			FSimpleDelegate::CreateLambda([this]()
+				{
+					if (ChangeCuePointOriginNotificationItem.IsValid())
+					{
+						ChangeCuePointOriginNotificationItem->SetCompletionState(SNotificationItem::CS_None);
+						ChangeCuePointOriginNotificationItem->Fadeout();
+					}
+				})
+		));
+
+		ChangeCuePointOriginNotificationItem = FSlateNotificationManager::Get().AddNotification(Info);
+
+		if (ChangeCuePointOriginNotificationItem.IsValid())
+		{
+			ChangeCuePointOriginNotificationItem->SetCompletionState(SNotificationItem::CS_Pending);
+		}
+	}
+}
+
 bool FWaveformEditor::InitializeAudioComponent()
 {
 	if (!ensure(SoundWave))
@@ -1087,6 +1190,10 @@ bool FWaveformEditor::CanExecuteReimport() const
 
 bool FWaveformEditor::BindCommands()
 {
+	check(TransportController.IsValid());
+	check(ZoomManager.IsValid());
+	check(WaveWriter.IsValid());
+
 	const FWaveformEditorCommands& Commands = FWaveformEditorCommands::Get();
 
 	ToolkitCommands->MapAction(
@@ -1371,12 +1478,6 @@ void FWaveformEditor::NotifyPostChange(const FPropertyChangedEvent& PropertyChan
 	WaveformView.DataProvider->UpdateRenderElements();
 	TransformationChainConfig = SoundWave->GetTransformationChainConfig();
 
-	if (PropertyChangedEvent.ChangeType == EPropertyChangeType::ValueSet)
-	{
-		UpdateTransportState();
-
-		bIsInteractingWithTransformations = false;
-	}
 	TObjectPtr<UWaveformTransformationMarkers> MarkerTransformation = GetMarkerTransformation();
 
 	if (MarkerTransformation != nullptr && !bCueChangeRegisteredByWaveformEditor)
@@ -1387,6 +1488,16 @@ void FWaveformEditor::NotifyPostChange(const FPropertyChangedEvent& PropertyChan
 	else if (MarkerTransformation == nullptr)
 	{
 		bCueChangeRegisteredByWaveformEditor = false;
+		bMarkerTransformationIsActive = false;
+	}
+
+	if (PropertyChangedEvent.ChangeType == EPropertyChangeType::ValueSet)
+	{
+		UpdateTransportState();
+
+		bIsInteractingWithTransformations = false;
+
+		PromptToChangeCuePointOrigin();
 	}
 }
 
diff --git a/Engine/Plugins/Editor/WaveformEditor/Source/WaveformEditor/Public/WaveformEditor.h b/Engine/Plugins/Editor/WaveformEditor/Source/WaveformEditor/Public/WaveformEditor.h
--- a/Engine/Plugins/Editor/WaveformEditor/Source/WaveformEditor/Public/WaveformEditor.h
+++ b/Engine/Plugins/Editor/WaveformEditor/Source/WaveformEditor/Public/WaveformEditor.h
@@ -109,7 +109,7 @@ private:
 	UE_API TObjectPtr<UWaveformTransformationTrimFade> GetOrAddTrimFadeTransformation();
 
 	UE_API TObjectPtr<UWaveformTransformationMarkers> GetOrAddMarkerTransformation();
-	UE_API TObjectPtr<UWaveformTransformationMarkers> GetMarkerTransformation();
+	UE_API TObjectPtr<UWaveformTransformationMarkers> GetMarkerTransformation() const;
 	UE_API TObjectPtr<UWaveformTransformationMarkers> AddMarkerTransformation();
 	UE_API void ToggleFadeIn();
 	UE_API bool CanFadeIn();
@@ -125,6 +125,9 @@ private:
 	UE_API void ModifyMarkerLoopRegion(ELoopModificationControls Modification);
 	UE_API void CycleMarkerLoopRegion(ELoopModificationControls Modification);
 
+	UE_API void SetCuePointOrigin(TObjectPtr<UWaveformTransformationMarkers> MarkerTransformation);
+	UE_API void PromptToChangeCuePointOrigin();
+
 	FTransformedWaveformView  WaveformView;
 
 	/** Exports the edited waveform to a new asset */
@@ -190,7 +193,10 @@ private:
 	UE_API FText GetExportButtonToolTip() const;
 
 	FDelegateHandle OnMapOpenedHandle;
-	TSharedPtr<SNotificationItem> ReopenNotificationItem;
+	TSharedPtr<SNotificationItem> ReopenNotificationItem = nullptr;
+	TSharedPtr<SNotificationItem> ChangeCuePointOriginNotificationItem = nullptr;
+	
+	bool bMarkerTransformationIsActive = false;
 };
 
 #undef UE_API
diff --git a/Engine/Source/Editor/AudioEditor/Private/Factories/SoundFactory.cpp b/Engine/Source/Editor/AudioEditor/Private/Factories/SoundFactory.cpp
--- a/Engine/Source/Editor/AudioEditor/Private/Factories/SoundFactory.cpp
+++ b/Engine/Source/Editor/AudioEditor/Private/Factories/SoundFactory.cpp
@@ -662,11 +662,6 @@ UObject* USoundFactory::CreateObject
 		// Store the current file path and timestamp for re-import purposes
 		Sound->AssetImportData->Update(CurrentFilename);
 
-#if WITH_EDITORONLY_DATA
-		// Setup the cue points
-		Sound->CuePointOrigin = ESoundWaveCuePointOrigin::WaveFile;
-#endif //WITH_EDITORONLY_DATA
-
 		int TotalNum = WaveInfo.WaveCues.Num() + WaveInfo.WaveSampleLoops.Num();
 		Sound->CuePoints.Reset(TotalNum);
 		TSet<int32> WaveCuePointIDs;
@@ -724,6 +719,11 @@ UObject* USoundFactory::CreateObject
 			}
 		}
 
+#if WITH_EDITORONLY_DATA
+		// Setup the cue points
+		Sound->CuePointOrigin = Sound->CuePoints.IsEmpty() ? ESoundWaveCuePointOrigin::MarkerTransformation : ESoundWaveCuePointOrigin::WaveFile;
+#endif //WITH_EDITORONLY_DATA
+
 		// fail import if we found invalid sample loops
 		// each invalid sample loop is logged
 		if (FoundInvalidSampleLoops)
diff --git a/Engine/Source/Runtime/Engine/Classes/Sound/SoundWave.h b/Engine/Source/Runtime/Engine/Classes/Sound/SoundWave.h
--- a/Engine/Source/Runtime/Engine/Classes/Sound/SoundWave.h
+++ b/Engine/Source/Runtime/Engine/Classes/Sound/SoundWave.h
@@ -493,6 +493,9 @@ public:
 #if WITH_EDITOR
 	ENGINE_API TArray<FSoundWaveCuePoint> GetSoundWaveCuePoints() const;
 	ENGINE_API void SetSoundWaveCuePoints(const TArray<FSoundWaveCuePoint>& InCuePoints);
+
+	ENGINE_API const ESoundWaveCuePointOrigin GetCuePointOrigin() const;
+	ENGINE_API void SetCuePointOrigin(const ESoundWaveCuePointOrigin InCuePointOrigin);
 #endif //WITH_EDITOR
 
 	/** Filters for the cue points that _are_ loop regions and returns those as a new array */
diff --git a/Engine/Source/Runtime/Engine/Private/SoundWave.cpp b/Engine/Source/Runtime/Engine/Private/SoundWave.cpp
--- a/Engine/Source/Runtime/Engine/Private/SoundWave.cpp
+++ b/Engine/Source/Runtime/Engine/Private/SoundWave.cpp
@@ -995,7 +995,7 @@ USoundWave::USoundWave(const FObjectInitializer& ObjectInitializer)
 	EnvelopeFollowerReleaseTime = 100;
 	bEnableCloudStreaming = false;
 
-	CuePointOrigin = ESoundWaveCuePointOrigin::WaveFile;
+	CuePointOrigin = ESoundWaveCuePointOrigin::MarkerTransformation;
 #endif
 
 	bCachedSampleRateFromPlatformSettings = false;
@@ -1537,6 +1537,16 @@ void USoundWave::SetSoundWaveCuePoints(const TArray<FSoundWaveCuePoint>& InCuePo
 	CuePoints = InCuePoints;
 	SetAllCuePoints(InCuePoints);
 }
+
+const ESoundWaveCuePointOrigin USoundWave::GetCuePointOrigin() const
+{
+	return CuePointOrigin;
+}
+
+void USoundWave::SetCuePointOrigin(const ESoundWaveCuePointOrigin InCuePointOrigin)
+{
+	CuePointOrigin = InCuePointOrigin;
+}
 #endif //WITH_EDITOR
 
 TArray<FSoundWaveCuePoint> USoundWave::GetLoopRegions() const
