Waveform Editor: re-ordered normalize and marker transformations

diff --git a/Engine/Plugins/Editor/WaveformEditor/Source/WaveformTransformations/Private/WaveformTransformationMarkers.cpp b/Engine/Plugins/Editor/WaveformEditor/Source/WaveformTransformations/Private/WaveformTransformationMarkers.cpp
--- a/Engine/Plugins/Editor/WaveformEditor/Source/WaveformTransformations/Private/WaveformTransformationMarkers.cpp
+++ b/Engine/Plugins/Editor/WaveformEditor/Source/WaveformTransformations/Private/WaveformTransformationMarkers.cpp
@@ -23,10 +23,8 @@ void FWaveTransformationMarkers::ProcessAudio(Audio::FWaveformTransformationWave
 	check(StartLoopTime >= 0);
 	check(InOutWaveInfo.NumChannels > 0);
 
-	const int64 StartSampleOffset = InOutWaveInfo.StartFrameOffset;
-
 	int64 StartSample = FMath::RoundToInt32(StartLoopTime * InOutWaveInfo.SampleRate) * InOutWaveInfo.NumChannels;
-	StartSample = FMath::Max(0, StartSample - StartSampleOffset);
+	StartSample = FMath::Max(0, StartSample);
 
 	if (StartSample > InputAudio.Num())
 	{
@@ -38,7 +36,7 @@ void FWaveTransformationMarkers::ProcessAudio(Audio::FWaveformTransformationWave
 	if (EndLoopTime > 0.f)
 	{
 		const int64 EndFrame = FMath::RoundToInt32(EndLoopTime * InOutWaveInfo.SampleRate) * InOutWaveInfo.NumChannels;
-		EndSample = FMath::Max(StartSample + 1, EndFrame - 1 - StartSampleOffset);
+		EndSample = FMath::Max(StartSample + 1, EndFrame - 1);
 
 		// EndLoopTime can be beyond the length of the file if there is a Trim
 		if (EndSample > InputAudio.Num() - 1)
@@ -114,17 +112,18 @@ void UWaveformTransformationMarkers::PostEditChangeProperty(FPropertyChangedEven
 					check(SampleRate > 0.f);
 					check(EndLoopTime > 0.0);
 
-					const int64 StartLoopFramePos = StartLoopTime * SampleRate;
-					const int64 EndLoopFramePos = EndLoopTime * SampleRate;
+					const int64 StartLoopFramePos = (StartLoopTime * SampleRate);
+					const int64 EndLoopFramePos = (EndLoopTime * SampleRate);
 
 					if (CuePoint.FramePosition != StartLoopFramePos)
 					{
-						StartLoopTime = static_cast<float>(CuePoint.FramePosition) / SampleRate;
+						StartLoopTime = static_cast<float>(FMath::Max(0, CuePoint.FramePosition - StartFrameOffset)) / SampleRate;
+						check(StartLoopTime >= 0);
 					}
 
 					if (CuePoint.FramePosition + CuePoint.FrameLength != EndLoopFramePos)
 					{
-						EndLoopTime = static_cast<float>(CuePoint.FramePosition + CuePoint.FrameLength) / SampleRate;
+						EndLoopTime = static_cast<float>(CuePoint.FramePosition + CuePoint.FrameLength - StartFrameOffset) / SampleRate;
 					}
 				}
 
diff --git a/Engine/Plugins/Editor/WaveformEditor/Source/WaveformTransformations/Private/WaveformTransformationNormalize.cpp b/Engine/Plugins/Editor/WaveformEditor/Source/WaveformTransformations/Private/WaveformTransformationNormalize.cpp
--- a/Engine/Plugins/Editor/WaveformEditor/Source/WaveformTransformations/Private/WaveformTransformationNormalize.cpp
+++ b/Engine/Plugins/Editor/WaveformEditor/Source/WaveformTransformations/Private/WaveformTransformationNormalize.cpp
@@ -140,6 +140,7 @@ FWaveTransformationNormalize::FWaveTransformationNormalize(float InTarget, float
 
 	const float TargetOffset = Target - PeakDecibelValue;
 	const float MakeupGain = FMath::Clamp(TargetOffset, -MaxGain, MaxGain);
+	InOutWaveInfo.StartFrameOffset = 0;
 
 	if (MakeupGain != 0.f)
 	{
diff --git a/Engine/Plugins/Editor/WaveformEditor/Source/WaveformTransformations/Public/WaveformTransformationMarkers.h b/Engine/Plugins/Editor/WaveformEditor/Source/WaveformTransformations/Public/WaveformTransformationMarkers.h
--- a/Engine/Plugins/Editor/WaveformEditor/Source/WaveformTransformations/Public/WaveformTransformationMarkers.h
+++ b/Engine/Plugins/Editor/WaveformEditor/Source/WaveformTransformations/Public/WaveformTransformationMarkers.h
@@ -102,7 +102,6 @@ public:
 
 	UE_API void UpdateConfiguration(FWaveTransformUObjectConfiguration& InOutConfiguration) override;
 	UE_API void OverwriteTransformation() override;
-	virtual constexpr Audio::ETransformationPriority GetTransformationPriority() const { return Audio::ETransformationPriority::Low; }
 
 	UE_API void ModifyMarkerLoopRegion(ELoopModificationControls Modification) const;
 	UE_API void CycleMarkerLoopRegion(ELoopModificationControls Modification) const;
diff --git a/Engine/Plugins/Editor/WaveformEditor/Source/WaveformTransformations/Public/WaveformTransformationNormalize.h b/Engine/Plugins/Editor/WaveformEditor/Source/WaveformTransformations/Public/WaveformTransformationNormalize.h
--- a/Engine/Plugins/Editor/WaveformEditor/Source/WaveformTransformations/Public/WaveformTransformationNormalize.h
+++ b/Engine/Plugins/Editor/WaveformEditor/Source/WaveformTransformations/Public/WaveformTransformationNormalize.h
@@ -49,6 +49,7 @@ class UWaveformTransformationNormalize : public UWaveformTransformationBase
 	
 public:
 	UE_API virtual Audio::FTransformationPtr CreateTransformation() const override;
+	virtual constexpr Audio::ETransformationPriority GetTransformationPriority() const { return Audio::ETransformationPriority::Low; }
 };
 
 #undef UE_API
diff --git a/Engine/Plugins/Editor/WaveformEditor/Source/WaveformTransformationsWidgets/Private/WaveformTransformationMarkersRenderer.cpp b/Engine/Plugins/Editor/WaveformEditor/Source/WaveformTransformationsWidgets/Private/WaveformTransformationMarkersRenderer.cpp
--- a/Engine/Plugins/Editor/WaveformEditor/Source/WaveformTransformationsWidgets/Private/WaveformTransformationMarkersRenderer.cpp
+++ b/Engine/Plugins/Editor/WaveformEditor/Source/WaveformTransformationsWidgets/Private/WaveformTransformationMarkersRenderer.cpp
@@ -852,8 +852,10 @@ void FWaveformTransformationMarkerRenderer::PreviewSelectedLoop()
 			check(CuePoint.FrameLength > 0);
 
 			// Set Loop preview handles
-			const float StartLoopPosInSeconds = static_cast<float>(CuePoint.FramePosition) / TransformationWaveInfo.SampleRate;
-			const float EndLoopPosInSeconds = static_cast<float>(LoopEndFramePosition) / TransformationWaveInfo.SampleRate;
+			const float StartLoopPosInSeconds = static_cast<float>(FMath::Max(0, CuePoint.FramePosition - StartFrameOffset)) / TransformationWaveInfo.SampleRate;
+			const float EndLoopPosInSeconds = static_cast<float>(LoopEndFramePosition - StartFrameOffset) / TransformationWaveInfo.SampleRate;
+
+			check(StartLoopPosInSeconds >= 0);
 
 			SetIsPreviewingLoopRegion(static_cast<double>(StartLoopPosInSeconds), static_cast<double>(EndLoopPosInSeconds), true);
 		}
