Removal of cached property handles and adding transformation to the transformation renderers.
[Waveform Editor] - Convert To Loop Region popup is sized incorrectly
[Waveform Editor] - Able to drag Isolated region of a Loop Region with a Start Trim

diff --git a/Engine/Plugins/Editor/WaveformEditor/Source/WaveformEditor/Private/TransformedWaveformViewFactory.cpp b/Engine/Plugins/Editor/WaveformEditor/Source/WaveformEditor/Private/TransformedWaveformViewFactory.cpp
--- a/Engine/Plugins/Editor/WaveformEditor/Source/WaveformEditor/Private/TransformedWaveformViewFactory.cpp
+++ b/Engine/Plugins/Editor/WaveformEditor/Source/WaveformEditor/Private/TransformedWaveformViewFactory.cpp
@@ -2,13 +2,10 @@
 
 #include "TransformedWaveformViewFactory.h"
 
-#include "Misc/NotifyHook.h"
-#include "PropertyEditorModule.h"
 #include "Sound/SoundWave.h"
 #include "SparseSampledSequenceTransportCoordinator.h"
 #include "SWaveformTransformationsOverlay.h"
 #include "TransformedWaveformView.h"
-#include "TransformedWaveformView.h"
 #include "WaveformEditorDetailsCustomization.h"
 #include "WaveformEditorSequenceDataProvider.h"
 #include "WaveformEditorZoomController.h"
@@ -31,9 +28,7 @@ void FTransformedWaveformViewFactory::Create()
 
 FTransformedWaveformView FTransformedWaveformViewFactory::GetTransformedView(TObjectPtr<USoundWave> SoundWaveToView, TSharedRef<FSparseSampledSequenceTransportCoordinator> TransportCoordinator, FNotifyHook* OwnerNotifyHook /*= nullptr*/, TSharedPtr<FWaveformEditorZoomController> InZoomController /*= nullptr*/)
 {
-	FOnTransformationsPropertiesRequired OnTransformationPropertiesRequired = OwnerNotifyHook ? SetUpOwnerNotifyHook(SoundWaveToView, OwnerNotifyHook) : FOnTransformationsPropertiesRequired();
-
-	TSharedPtr<FWaveformEditorSequenceDataProvider> WaveformDataProvider = MakeShared<FWaveformEditorSequenceDataProvider>(SoundWaveToView, OnTransformationPropertiesRequired);
+	TSharedPtr<FWaveformEditorSequenceDataProvider> WaveformDataProvider = MakeShared<FWaveformEditorSequenceDataProvider>(SoundWaveToView);
 	WaveformDataProvider->UpdateRenderElements();
 
 	FFixedSampledSequenceView SequenceView = WaveformDataProvider->RequestSequenceView(TRange<double>::Inclusive(0, 1));
@@ -56,49 +51,6 @@ FTransformedWaveformView FTransformedWaveformViewFactory::GetTransformedView(TOb
 
 }
 
-FOnTransformationsPropertiesRequired FTransformedWaveformViewFactory::SetUpOwnerNotifyHook(TObjectPtr<USoundWave> SoundWaveToView, FNotifyHook* InNotifyHook)
-{
-	TSharedPtr<IDetailsView> PropertiesPropagator = SetUpTransformationsPropertiesPropagator(SoundWaveToView, InNotifyHook);
-
-	FOnTransformationsPropertiesRequired OnTransformationPropertiesRequired = FOnTransformationsPropertiesRequired::CreateLambda([this, SoundWaveToView, PropertiesPropagator](FTransformationsToPropertiesArray& InObjToPropsMap)
-	{
-		TSharedPtr<FWaveformTransformationsDetailsProvider> TransformationsDetailsProvider = MakeShared<FWaveformTransformationsDetailsProvider>();
-
-		FOnGetDetailCustomizationInstance ProviderInstance = FOnGetDetailCustomizationInstance::CreateLambda([&]()
-			{
-				return TransformationsDetailsProvider.ToSharedRef();
-			}
-		);
-
-		PropertiesPropagator->RegisterInstancedCustomPropertyLayout(SoundWaveToView->GetClass(), ProviderInstance);
-		PropertiesPropagator->ForceRefresh();
-
-		for (FTransformationToPropertiesPair& ObjToPropsPair : InObjToPropsMap)
-		{
-			TransformationsDetailsProvider->GetHandlesForUObjectProperties(ObjToPropsPair.Key, ObjToPropsPair.Value);
-		}
-
-		TransformationsDetailsProvider.Reset();
-	});
-
-	return OnTransformationPropertiesRequired;
-
-}
-
-TSharedPtr<IDetailsView> FTransformedWaveformViewFactory::SetUpTransformationsPropertiesPropagator(TObjectPtr<USoundWave> SoundWaveToView, FNotifyHook* InNotifyHook)
-{
-	check(InNotifyHook)
-
-	FPropertyEditorModule& PropertyModule = FModuleManager::LoadModuleChecked<FPropertyEditorModule>("PropertyEditor");
-	FDetailsViewArgs Args;
-	Args.bHideSelectionTip = true;
-	Args.NotifyHook = InNotifyHook;
-
-	TSharedPtr<IDetailsView> TransformationsPropertiesPropagator = PropertyModule.CreateDetailView(Args);
-	TransformationsPropertiesPropagator->SetObject(SoundWaveToView);
-	return TransformationsPropertiesPropagator;
-}
-
 void FTransformedWaveformViewFactory::SetUpWaveformPanelInteractions(TSharedRef<STransformedWaveformViewPanel> WaveformPanel, TSharedRef<FSparseSampledSequenceTransportCoordinator> TransportCoordinator, TSharedPtr<FWaveformEditorZoomController> InZoomController /*= nullptr*/, TSharedPtr<SWaveformTransformationsOverlay> InTransformationsOverlay /*=nullptr*/)
 {
 	FPointerEventHandler HandlePlayheadOverlayMouseButtonUp = FPointerEventHandler::CreateLambda([TransportCoordinator](const FGeometry& Geometry, const FPointerEvent& MouseEvent)
diff --git a/Engine/Plugins/Editor/WaveformEditor/Source/WaveformEditor/Private/TransformedWaveformViewFactory.h b/Engine/Plugins/Editor/WaveformEditor/Source/WaveformEditor/Private/TransformedWaveformViewFactory.h
--- a/Engine/Plugins/Editor/WaveformEditor/Source/WaveformEditor/Private/TransformedWaveformViewFactory.h
+++ b/Engine/Plugins/Editor/WaveformEditor/Source/WaveformEditor/Private/TransformedWaveformViewFactory.h
@@ -39,8 +39,6 @@ public:
 	UE_API FTransformedWaveformView GetTransformedView(TObjectPtr<USoundWave> SoundWaveToView, TSharedRef<FSparseSampledSequenceTransportCoordinator> TransportCoordinator, FNotifyHook* OwnerNotifyHook = nullptr, TSharedPtr<FWaveformEditorZoomController> InZoomController = nullptr);
 
 private:
-	UE_API FOnTransformationsPropertiesRequired SetUpOwnerNotifyHook(TObjectPtr<USoundWave> SoundWaveToView, FNotifyHook* InNotifyHook);
-	UE_API TSharedPtr<IDetailsView> SetUpTransformationsPropertiesPropagator(TObjectPtr<USoundWave> SoundWaveToView, FNotifyHook* InNotifyHook);
 	UE_API void SetUpWaveformPanelInteractions(TSharedRef<STransformedWaveformViewPanel> WaveformPanel, TSharedRef<FSparseSampledSequenceTransportCoordinator> TransportCoordinator, TSharedPtr<FWaveformEditorZoomController> InZoomController = nullptr, TSharedPtr<SWaveformTransformationsOverlay> InTransformationsOverlay = nullptr);
 	UE_API FReply HandleTimeRulerInteraction(const TransformedWaveformViewFactory::EReceivedInteractionType MouseInteractionType, TSharedRef<FSparseSampledSequenceTransportCoordinator> TransportCoordinator, const TSharedRef<SWidget> TimeRulerWidget, const FPointerEvent& MouseEvent, const FGeometry& Geometry);
 	
diff --git a/Engine/Plugins/Editor/WaveformEditor/Source/WaveformEditor/Private/WaveformEditor.cpp b/Engine/Plugins/Editor/WaveformEditor/Source/WaveformEditor/Private/WaveformEditor.cpp
--- a/Engine/Plugins/Editor/WaveformEditor/Source/WaveformEditor/Private/WaveformEditor.cpp
+++ b/Engine/Plugins/Editor/WaveformEditor/Source/WaveformEditor/Private/WaveformEditor.cpp
@@ -94,79 +94,90 @@ bool FWaveformEditor::Init(const EToolkitMode::Type Mode, const TSharedPtr<ITool
 			bToolbarFocusable,
 			bUseSmallIcons);
 
-			// Initialize waveform view to with playhead at start of active transformations.
-			check(WaveformView.DataProvider != nullptr);
-			WaveformView.DataProvider->GenerateLayersChain();
-			WaveformView.DataProvider->UpdateRenderElements();
-			PlaybackTimeBeforeTransformInteraction = 0;
-			StartTimeBeforeTransformInteraction = TransformationChainConfig.StartTime; // StartTime=0, EndTime=-1 if no transformations
-			check(TransportController != nullptr);
-			TransportController->CacheStartTime(PlaybackTimeBeforeTransformInteraction);
-			check(TransportCoordinator != nullptr);
-			TransportCoordinator->SetProgressRatio(0.0f);
+		// Initialize waveform view to with playhead at start of active transformations.
+		check(WaveformView.DataProvider != nullptr);
+		WaveformView.DataProvider->GenerateLayersChain();
+		WaveformView.DataProvider->UpdateRenderElements();
+		PlaybackTimeBeforeTransformInteraction = 0;
+		StartTimeBeforeTransformInteraction = TransformationChainConfig.StartTime; // StartTime=0, EndTime=-1 if no transformations
+		check(TransportController != nullptr);
+		TransportController->CacheStartTime(PlaybackTimeBeforeTransformInteraction);
+		check(TransportCoordinator != nullptr);
+		TransportCoordinator->SetProgressRatio(0.0f);
+
+		TWeakPtr<FWaveformEditor> WeakWaveformEditor = SharedThis(this);
+		TObjectPtr<UWaveformTransformationMarkers> MarkerTransformation = GetMarkerTransformation();
+
+		if (MarkerTransformation)
+		{
+			bMarkerTransformationIsActive = true;
 
-			if (GetMarkerTransformation())
+			if (!bCueChangeRegisteredByWaveformEditor)
 			{
-				bMarkerTransformationIsActive = true;
+				MarkerTransformation->Markers->CueChanged.BindLambda([WeakWaveformEditor]() 
+					{ 
+						if (TSharedPtr<FWaveformEditor> PinnedEditor = WeakWaveformEditor.Pin())
+						{
+							PinnedEditor->RegenerateTransformations();
+						}
+					});
+
+				bCueChangeRegisteredByWaveformEditor = true;
 			}
 
-			OnMapOpenedHandle = FEditorDelegates::OnMapOpened.AddLambda([this, SoundWaveToEdit](const FString&, bool)
-				{
-					CloseWindow(EAssetEditorCloseReason::EditorRefreshRequested);
-
-					FNotificationInfo Info(FText::FromString(TEXT("Do you want to reopen the sound wave for editing?")));
-					Info.bFireAndForget = true;
-					Info.ExpireDuration = 15.0f;
-
-					Info.ButtonDetails.Add(FNotificationButtonInfo(
-						FText::FromString(TEXT("Yes")),
-						FText::FromString(TEXT("Open the SoundWave asset")),
-						FSimpleDelegate::CreateLambda([this, SoundWaveToEdit]() 
-							{
-								bool bSuccess = false;
-								if (GEditor && SoundWaveToEdit)
-								{
-									TArray<UObject*> AssetsToOpen = { SoundWaveToEdit };
-									if (UAssetEditorSubsystem* AssetEditorSubsystem = GEditor->GetEditorSubsystem<UAssetEditorSubsystem>())
-									{
-										bSuccess = AssetEditorSubsystem->OpenEditorForAssets(AssetsToOpen);
-									}
-								}
-								if (ReopenNotificationItem.IsValid())
-								{
-									if (bSuccess)
-									{
-										ReopenNotificationItem->SetCompletionState(SNotificationItem::CS_Success);
-									}
-									else
-									{
-										ReopenNotificationItem->SetCompletionState(SNotificationItem::CS_Fail);
-									}
-									ReopenNotificationItem->Fadeout();
-								}
-							})
-					));
-
-					Info.ButtonDetails.Add(FNotificationButtonInfo(
-						FText::FromString(TEXT("No")),
-						FText::FromString(TEXT("Dismiss this message")),
-						FSimpleDelegate::CreateLambda([this]()
-							{
-								if (ReopenNotificationItem.IsValid())
-								{
-									ReopenNotificationItem->SetCompletionState(SNotificationItem::CS_None);
-									ReopenNotificationItem->Fadeout();
-								}
-							})
-					));
-
-					ReopenNotificationItem = FSlateNotificationManager::Get().AddNotification(Info);
-					
-					if (ReopenNotificationItem.IsValid())
+			if (!MarkerTransformation->OnTransformationChanged.IsBound())
+			{
+				MarkerTransformation->OnTransformationChanged.BindLambda([WeakWaveformEditor]()
 					{
-						ReopenNotificationItem->SetCompletionState(SNotificationItem::CS_Pending);
-					}
-				});
+						if (TSharedPtr<FWaveformEditor> PinnedEditor = WeakWaveformEditor.Pin())
+						{
+							PinnedEditor->UpdateTransformations();
+						}
+					});
+			}
+			if (!MarkerTransformation->OnTransformationRenderChanged.IsBound())
+			{
+				MarkerTransformation->OnTransformationRenderChanged.BindLambda([WeakWaveformEditor]()
+					{
+						if (TSharedPtr<FWaveformEditor> PinnedEditor = WeakWaveformEditor.Pin())
+						{
+							PinnedEditor->UpdateRenderElements();
+						}
+					});
+			}
+		}
+		else
+		{
+			bCueChangeRegisteredByWaveformEditor = false;
+			bMarkerTransformationIsActive = false;
+		}
+
+		TObjectPtr<UWaveformTransformationTrimFade> TrimFadeTransformation = GetTrimFadeTransformation();
+
+		if (TrimFadeTransformation != nullptr)
+		{
+			if (!TrimFadeTransformation->OnTransformationChanged.IsBound())
+			{
+				TrimFadeTransformation->OnTransformationChanged.BindLambda([WeakWaveformEditor]()
+					{
+						if (TSharedPtr<FWaveformEditor> PinnedEditor = WeakWaveformEditor.Pin())
+						{
+							PinnedEditor->UpdateTransformations();
+						}
+					});
+			}
+
+			if (!TrimFadeTransformation->OnTransformationRenderChanged.IsBound())
+			{
+				TrimFadeTransformation->OnTransformationRenderChanged.BindLambda([WeakWaveformEditor]()
+					{
+						if (TSharedPtr<FWaveformEditor> PinnedEditor = WeakWaveformEditor.Pin())
+						{
+							PinnedEditor->UpdateRenderElements();
+						}
+					});
+			}
+		}
 	}
 
 	return bIsInitialized;
@@ -174,26 +185,24 @@ bool FWaveformEditor::Init(const EToolkitMode::Type Mode, const TSharedPtr<ITool
 
 FWaveformEditor::~FWaveformEditor()
 {
-	if (OnMapOpenedHandle.IsValid())
-	{
-		FEditorDelegates::OnMapOpened.Remove(OnMapOpenedHandle);
-		OnMapOpenedHandle.Reset();
-	}
-
 	if (ChangeCuePointOriginNotificationItem.IsValid())
 	{
 		ChangeCuePointOriginNotificationItem->Fadeout();
 	}
 
-	if (FReimportManager::Instance())
-	{
-		FReimportManager::Instance()->OnPostReimport().RemoveAll(this);
-	}
-
 	TObjectPtr<UWaveformTransformationMarkers> MarkerTransformation = GetMarkerTransformation();
 
-	if (MarkerTransformation != nullptr)
+	if (MarkerTransformation)
 	{
+		if (MarkerTransformation->OnTransformationChanged.IsBound())
+		{
+			MarkerTransformation->OnTransformationChanged.Unbind();
+		}
+		if (MarkerTransformation->OnTransformationRenderChanged.IsBound())
+		{
+			MarkerTransformation->OnTransformationRenderChanged.Unbind();
+		}
+
 		if (bCueChangeRegisteredByWaveformEditor)
 		{
 			check(MarkerTransformation->Markers);
@@ -213,6 +222,25 @@ FWaveformEditor::~FWaveformEditor()
 			SoundWave->PostEditChange();
 		}
 	}
+
+	TObjectPtr<UWaveformTransformationTrimFade> TrimFadeTransformation = GetTrimFadeTransformation();
+
+	if (TrimFadeTransformation)
+	{
+		if (TrimFadeTransformation->OnTransformationChanged.IsBound())
+		{
+			TrimFadeTransformation->OnTransformationChanged.Unbind();
+		}
+		if (TrimFadeTransformation->OnTransformationRenderChanged.IsBound())
+		{
+			TrimFadeTransformation->OnTransformationRenderChanged.Unbind();
+		}
+	}
+
+	if (FReimportManager::Instance())
+	{
+		FReimportManager::Instance()->OnPostReimport().RemoveAll(this);
+	}
 }
 
 void FWaveformEditor::AddDefaultTransformations()
@@ -280,16 +308,7 @@ TObjectPtr<UWaveformTransformationTrimFade> FWaveformEditor::GetOrAddTrimFadeTra
 {
 	check(SoundWave);
 
-	TObjectPtr<UWaveformTransformationTrimFade> TrimFadeTransformation;
-
-	for (TObjectPtr<UWaveformTransformationBase> Transformation : SoundWave->Transformations)
-	{
-		TrimFadeTransformation = Cast<UWaveformTransformationTrimFade>(Transformation);
-		if (TrimFadeTransformation)
-		{
-			break;
-		}
-	}
+	TObjectPtr<UWaveformTransformationTrimFade> TrimFadeTransformation = GetTrimFadeTransformation();
 
 	if (TrimFadeTransformation == nullptr)
 	{
@@ -310,6 +329,27 @@ TObjectPtr<UWaveformTransformationTrimFade> FWaveformEditor::GetOrAddTrimFadeTra
 	return TrimFadeTransformation;
 }
 
+TObjectPtr<UWaveformTransformationTrimFade> FWaveformEditor::GetTrimFadeTransformation() const
+{
+	if (SoundWave == nullptr)
+	{
+		return nullptr;
+	}
+
+	// Relies on there being only one TrimFade Transformation on the soundwave
+	TObjectPtr<UWaveformTransformationBase>* TrimFadeTransformation = SoundWave->Transformations.FindByPredicate([](const TObjectPtr<UWaveformTransformationBase>& Transformation)
+		{
+			return Transformation.IsA<UWaveformTransformationTrimFade>();
+		});
+
+	if (TrimFadeTransformation)
+	{
+		return CastChecked<UWaveformTransformationTrimFade>(*TrimFadeTransformation);
+	}
+
+	return nullptr;
+}
+
 TObjectPtr<UWaveformTransformationMarkers> FWaveformEditor::GetOrAddMarkerTransformation()
 {
 	check(SoundWave);
@@ -332,15 +372,15 @@ TObjectPtr<UWaveformTransformationMarkers> FWaveformEditor::GetMarkerTransformat
 		return nullptr;
 	}
 
-	TObjectPtr<UWaveformTransformationMarkers> MarkersTransformation;
+	// Relies on there being only one Markers Transformation on the soundwave
+	TObjectPtr<UWaveformTransformationBase>* MarkersTransformation = SoundWave->Transformations.FindByPredicate([](const TObjectPtr<UWaveformTransformationBase>& Transformation)
+		{
+			return Transformation.IsA<UWaveformTransformationMarkers>();
+		});
 
-	for (TObjectPtr<UWaveformTransformationBase> Transformation : SoundWave->Transformations)
+	if (MarkersTransformation)
 	{
-		MarkersTransformation = Cast<UWaveformTransformationMarkers>(Transformation);
-		if (MarkersTransformation)
-		{
-			return MarkersTransformation;
-		}
+		return CastChecked<UWaveformTransformationMarkers>(*MarkersTransformation);
 	}
 
 	return nullptr;
@@ -579,13 +619,32 @@ void FWaveformEditor::SkipToNextMarker()
 
 void FWaveformEditor::RegenerateTransformations()
 {
+	check(SoundWave);
 	WaveformView.DataProvider->GenerateLayersChain();
 	WaveformView.DataProvider->UpdateRenderElements();
-	TransformationChainConfig = SoundWave->GetTransformationChainConfig();
-
+	TransformationChainConfig = SoundWave->UpdateTransformations();
+	
 	UpdateTransportState();
 }
 
+void FWaveformEditor::UpdateTransformations()
+{
+	check(TransportController);
+
+	bool bIsPlaying = TransportController->IsPlaying();
+	NotifyPostTransformationChange(EPropertyChangeType::ValueSet);
+
+	if (bIsPlaying)
+	{
+		TransportController->Play();
+	}
+}
+
+void FWaveformEditor::UpdateRenderElements()
+{
+	WaveformView.DataProvider->UpdateRenderElements();
+}
+
 void FWaveformEditor::ModifyMarkerLoopRegion(ELoopModificationControls Modification)
 {
 	check(SoundWave);
@@ -1478,19 +1537,78 @@ void FWaveformEditor::NotifyPostChange(const FPropertyChangedEvent& PropertyChan
 	WaveformView.DataProvider->UpdateRenderElements();
 	TransformationChainConfig = SoundWave->GetTransformationChainConfig();
 
+	TWeakPtr<FWaveformEditor> WeakWaveformEditor = SharedThis(this);
 	TObjectPtr<UWaveformTransformationMarkers> MarkerTransformation = GetMarkerTransformation();
 
-	if (MarkerTransformation != nullptr && !bCueChangeRegisteredByWaveformEditor)
+	if (MarkerTransformation != nullptr)
 	{
-		MarkerTransformation->Markers->CueChanged.BindLambda([this]() { RegenerateTransformations(); });
-		bCueChangeRegisteredByWaveformEditor = true;
+		if (!bCueChangeRegisteredByWaveformEditor)
+		{
+			MarkerTransformation->Markers->CueChanged.BindLambda([WeakWaveformEditor]()
+				{
+					if (TSharedPtr<FWaveformEditor> PinnedEditor = WeakWaveformEditor.Pin())
+					{
+						PinnedEditor->RegenerateTransformations();
+					}
+				});
+
+			bCueChangeRegisteredByWaveformEditor = true;
+		}
+
+		if (!MarkerTransformation->OnTransformationChanged.IsBound())
+		{
+			MarkerTransformation->OnTransformationChanged.BindLambda([WeakWaveformEditor]()
+				{
+					if (TSharedPtr<FWaveformEditor> PinnedEditor = WeakWaveformEditor.Pin())
+					{
+						PinnedEditor->UpdateTransformations();
+					}
+				});
+		}
+		if (!MarkerTransformation->OnTransformationRenderChanged.IsBound())
+		{
+			MarkerTransformation->OnTransformationRenderChanged.BindLambda([WeakWaveformEditor]()
+				{
+					if (TSharedPtr<FWaveformEditor> PinnedEditor = WeakWaveformEditor.Pin())
+					{
+						PinnedEditor->UpdateRenderElements();
+					}
+				});
+		}
 	}
-	else if (MarkerTransformation == nullptr)
+	else
 	{
 		bCueChangeRegisteredByWaveformEditor = false;
 		bMarkerTransformationIsActive = false;
 	}
 
+	TObjectPtr<UWaveformTransformationTrimFade> TrimFadeTransformation = GetTrimFadeTransformation();
+
+	if (TrimFadeTransformation != nullptr)
+	{
+		if (!TrimFadeTransformation->OnTransformationChanged.IsBound())
+		{
+			TrimFadeTransformation->OnTransformationChanged.BindLambda([WeakWaveformEditor]()
+				{
+					if (TSharedPtr<FWaveformEditor> PinnedEditor = WeakWaveformEditor.Pin())
+					{
+						PinnedEditor->UpdateTransformations();
+					}
+				});
+		}
+
+		if (!TrimFadeTransformation->OnTransformationRenderChanged.IsBound())
+		{
+			TrimFadeTransformation->OnTransformationRenderChanged.BindLambda([WeakWaveformEditor]()
+				{
+					if (TSharedPtr<FWaveformEditor> PinnedEditor = WeakWaveformEditor.Pin())
+					{
+						PinnedEditor->UpdateRenderElements();
+					}
+				});
+		}
+	}
+
 	if (PropertyChangedEvent.ChangeType == EPropertyChangeType::ValueSet)
 	{
 		UpdateTransportState();
diff --git a/Engine/Plugins/Editor/WaveformEditor/Source/WaveformEditor/Public/WaveformEditor.h b/Engine/Plugins/Editor/WaveformEditor/Source/WaveformEditor/Public/WaveformEditor.h
--- a/Engine/Plugins/Editor/WaveformEditor/Source/WaveformEditor/Public/WaveformEditor.h
+++ b/Engine/Plugins/Editor/WaveformEditor/Source/WaveformEditor/Public/WaveformEditor.h
@@ -107,10 +107,14 @@ private:
 	/** Transformation Functions*/
 	UE_API void NotifyPostTransformationChange(const EPropertyChangeType::Type& PropertyChangeType = EPropertyChangeType::ArrayAdd);
 	UE_API TObjectPtr<UWaveformTransformationTrimFade> GetOrAddTrimFadeTransformation();
+	// Relies on there being only one TrimFade Transformation on the soundwave
+	UE_API TObjectPtr<UWaveformTransformationTrimFade> GetTrimFadeTransformation() const;
 
 	UE_API TObjectPtr<UWaveformTransformationMarkers> GetOrAddMarkerTransformation();
+	// Relies on there being only one Marker Transformation on the soundwave
 	UE_API TObjectPtr<UWaveformTransformationMarkers> GetMarkerTransformation() const;
 	UE_API TObjectPtr<UWaveformTransformationMarkers> AddMarkerTransformation();
+
 	UE_API void ToggleFadeIn();
 	UE_API bool CanFadeIn();
 
@@ -120,6 +124,8 @@ private:
 	UE_API void DeleteMarker();
 	UE_API void SkipToNextMarker(); // Skips the playhead to the next marker
 	UE_API void RegenerateTransformations();
+	UE_API void UpdateTransformations();
+	UE_API void UpdateRenderElements();
 	UE_API void UpdateTransportState();
 
 	UE_API void ModifyMarkerLoopRegion(ELoopModificationControls Modification);
@@ -192,7 +198,6 @@ private:
 	UE_API FText GetReimportButtonToolTip() const;
 	UE_API FText GetExportButtonToolTip() const;
 
-	FDelegateHandle OnMapOpenedHandle;
 	TSharedPtr<SNotificationItem> ReopenNotificationItem = nullptr;
 	TSharedPtr<SNotificationItem> ChangeCuePointOriginNotificationItem = nullptr;
 	
diff --git a/Engine/Plugins/Editor/WaveformEditor/Source/WaveformEditorWidgets/Private/WaveformEditorSequenceDataProvider.cpp b/Engine/Plugins/Editor/WaveformEditor/Source/WaveformEditorWidgets/Private/WaveformEditorSequenceDataProvider.cpp
--- a/Engine/Plugins/Editor/WaveformEditor/Source/WaveformEditorWidgets/Private/WaveformEditorSequenceDataProvider.cpp
+++ b/Engine/Plugins/Editor/WaveformEditor/Source/WaveformEditorWidgets/Private/WaveformEditorSequenceDataProvider.cpp
@@ -9,12 +9,9 @@
 #include "Sound/SoundWave.h"
 #include "WaveformTransformationDurationRenderer.h"
 
-FWaveformEditorSequenceDataProvider::FWaveformEditorSequenceDataProvider(
-	TObjectPtr<USoundWave> InSoundWave, 
-	FOnTransformationsPropertiesRequired InTransformationChainPropertiesHandler)
+FWaveformEditorSequenceDataProvider::FWaveformEditorSequenceDataProvider(TObjectPtr<USoundWave> InSoundWave)
 	: SoundWaveToRender(InSoundWave)
 	, LayersFactory(MakeUnique<FWaveformTransformationRenderLayerFactory>())
-	, ChainPropertiesHandler(InTransformationChainPropertiesHandler)
 {
 	GenerateLayersChain();	
 
@@ -54,37 +51,51 @@ FFixedSampledSequenceView FWaveformEditorSequenceDataProvider::RequestSequenceVi
 void FWaveformEditorSequenceDataProvider::GenerateLayersChain()
 {
 	check(SoundWaveToRender);
-		
-	TransformationsToRender.Empty();
-	RenderLayers.Empty();
 
-	if (SoundWaveToRender->Transformations.Num() > 0)
+	if (SoundWaveToRender->Transformations.IsEmpty())
 	{
-		FTransformationsToPropertiesArray TransformPropertiesMap;
+		TransformationsToRender.Empty();
+		RenderLayers.Empty();
+
+		OnLayersChainGenerated.Broadcast(RenderLayers.GetData(), RenderLayers.Num());
+	}
 
-		for (const TObjectPtr<UWaveformTransformationBase>& Transformation : SoundWaveToRender->Transformations)
+	TArray<TObjectPtr<UWaveformTransformationBase>> TempTransformationsToRender;
+	TArray<FTransformationRenderLayerInfo> TempRenderLayers;
+
+	for (const TObjectPtr<UWaveformTransformationBase>& Transformation : SoundWaveToRender->Transformations)
+	{
+		if (Transformation != nullptr)
 		{
-			if (Transformation != nullptr)
+			TempTransformationsToRender.Add(Transformation);
+			TSharedPtr<IWaveformTransformationRenderer> TransformationUI = LayersFactory->Create(Transformation);
+			FTransformationRenderLayerInfo RenderLayerInfo = FTransformationRenderLayerInfo(TransformationUI, FTransformationLayerConstraints(0.f, 1.f));
+			TempRenderLayers.Add(RenderLayerInfo);
+		}
+	}
+
+	bool bTransformationIsSame = true;
+	constexpr SIZE_T AccountForDurationLayer = 1;
+	if (TempTransformationsToRender.Num() == TransformationsToRender.Num() && TempRenderLayers.Num() == RenderLayers.Num() - AccountForDurationLayer)
+	{
+		for (int32 Index = 0; Index < TempTransformationsToRender.Num(); ++Index)
+		{
+			if (TempTransformationsToRender[Index] != TransformationsToRender[Index])
 			{
-				TArray<TSharedRef<IPropertyHandle>> TransformationPropertiesArray;
-				TransformPropertiesMap.Emplace(Transformation, TransformationPropertiesArray);
+				bTransformationIsSame = false;
 			}
 		}
 
-		ChainPropertiesHandler.Execute(TransformPropertiesMap);
-		
-
-		for (FTransformationToPropertiesPair& ObjPropsPair : TransformPropertiesMap)
+		if (bTransformationIsSame)
 		{
-			TransformationsToRender.Add(ObjPropsPair.Key);
-			TSharedPtr<IWaveformTransformationRenderer> TransformationUI = LayersFactory->Create(ObjPropsPair.Key, ObjPropsPair.Value);
-			FTransformationRenderLayerInfo RenderLayerInfo = FTransformationRenderLayerInfo(TransformationUI, FTransformationLayerConstraints(0.f, 1.f));
-			RenderLayers.Add(RenderLayerInfo);
+			return;
 		}
-
-		CreateDurationHighlightLayer();
 	}
 
+	TransformationsToRender = MoveTemp(TempTransformationsToRender);
+	RenderLayers = MoveTemp(TempRenderLayers);
+
+	CreateDurationHighlightLayer();
 	OnLayersChainGenerated.Broadcast(RenderLayers.GetData(), RenderLayers.Num());
 }
 
@@ -97,8 +108,11 @@ void FWaveformEditorSequenceDataProvider::UpdateRenderElements()
 void FWaveformEditorSequenceDataProvider::GenerateSequenceDataInternal()
 {
 	check(SoundWaveToRender);
+	check(ImportedNumChannels > 0);
 	
 	int64 NumWaveformSamples = ImportedRawPCMData.Num() * sizeof(uint8) / sizeof(int16);
+	check(NumWaveformSamples >= ImportedNumChannels);
+
 	NumOriginalWaveformFrames = NumWaveformSamples / ImportedNumChannels;
 	int64 FirstEditedSample = 0;
 	int64 LastEditedSample = NumWaveformSamples;
@@ -142,7 +156,7 @@ void FWaveformEditorSequenceDataProvider::GenerateSequenceDataInternal()
 				FirstEditedSample += TransformationInfo.StartFrameOffset;
 				if (FirstEditedSample >= OutputBuffer.Num())
 				{
-					FirstEditedSample = OutputBuffer.Num() - 1;
+					FirstEditedSample = OutputBuffer.Num() - ImportedNumChannels;
 				}
 
 				//Prevent a lower priority transformation from setting the LastEditedSample past the end of higher priority transformations
@@ -153,12 +167,15 @@ void FWaveformEditorSequenceDataProvider::GenerateSequenceDataInternal()
 					LastEditedSample = TransformationInfo.NumEditedSamples <= 0 ? LastEditedSample : FirstEditedSample + TransformationInfo.NumEditedSamples;
 				}
 			
-				if(LastEditedSample <= FirstEditedSample)
+				int32 EditedBufferSize = LastEditedSample - FirstEditedSample;
+				
+				if (EditedBufferSize <= ImportedNumChannels)
 				{
-					LastEditedSample = FirstEditedSample + 1;
+					EditedBufferSize = ImportedNumChannels;
+					FirstEditedSample = FMath::Clamp(FirstEditedSample, 0, NumWaveformSamples - ImportedNumChannels);
 				}
 
-				const int32 EditedBufferSize = LastEditedSample - FirstEditedSample;
+				check(FirstEditedSample + EditedBufferSize <= NumWaveformSamples);
 				FMemory::Memcpy(&OutputBuffer.GetData()[FirstEditedSample], TransformationsBuffer.GetData(), EditedBufferSize * sizeof(float));
 			}
 
diff --git a/Engine/Plugins/Editor/WaveformEditor/Source/WaveformEditorWidgets/Private/WaveformTransformationRenderLayerFactory.cpp b/Engine/Plugins/Editor/WaveformEditor/Source/WaveformEditorWidgets/Private/WaveformTransformationRenderLayerFactory.cpp
--- a/Engine/Plugins/Editor/WaveformEditor/Source/WaveformEditorWidgets/Private/WaveformTransformationRenderLayerFactory.cpp
+++ b/Engine/Plugins/Editor/WaveformEditor/Source/WaveformEditorWidgets/Private/WaveformTransformationRenderLayerFactory.cpp
@@ -17,6 +17,11 @@ TSharedPtr<IWaveformTransformationRenderer> FWaveformTransformationRenderLayerFa
 	{
 		Instantiator->CheckCallable();
 		OutRenderer = (*Instantiator)();
+
+		if (OutRenderer)
+		{
+			OutRenderer->SetWaveformTransformation(InTransformationToRender);
+		}
 	}
 
 	return OutRenderer;
@@ -28,8 +33,10 @@ TSharedPtr<IWaveformTransformationRenderer> FWaveformTransformationRenderLayerFa
 
 	if (OutRenderer)
 	{
+		PRAGMA_DISABLE_DEPRECATION_WARNINGS
 		OutRenderer->SetPropertyHandles(TransformationsProperties);
+		PRAGMA_ENABLE_DEPRECATION_WARNINGS
 	}
 
 	return OutRenderer;
-}
+}
\ No newline at end of file
diff --git a/Engine/Plugins/Editor/WaveformEditor/Source/WaveformEditorWidgets/Public/WaveformEditorSequenceDataProvider.h b/Engine/Plugins/Editor/WaveformEditor/Source/WaveformEditorWidgets/Public/WaveformEditorSequenceDataProvider.h
--- a/Engine/Plugins/Editor/WaveformEditor/Source/WaveformEditorWidgets/Public/WaveformEditorSequenceDataProvider.h
+++ b/Engine/Plugins/Editor/WaveformEditor/Source/WaveformEditorWidgets/Public/WaveformEditorSequenceDataProvider.h
@@ -55,10 +55,7 @@ class FWaveformEditorSequenceDataProvider : public IFixedSampledSequenceViewProv
 {
 public:
 
-	UE_API FWaveformEditorSequenceDataProvider(
-		TObjectPtr<USoundWave> InSoundWave,
-		FOnTransformationsPropertiesRequired InTransformationChainPropertiesHandler
-	);
+	UE_API FWaveformEditorSequenceDataProvider(TObjectPtr<USoundWave> InSoundWave);
 
 	UE_API virtual FFixedSampledSequenceView RequestSequenceView(const TRange<double> DataRatioRange) override;
 
@@ -108,8 +105,6 @@ private:
 	TUniquePtr<FWaveformTransformationRenderLayerFactory> LayersFactory = nullptr;
 	TSharedPtr<FWaveformTransformationDurationRenderer> DurationHiglightLayer = nullptr;
 
-	FOnTransformationsPropertiesRequired ChainPropertiesHandler;
-
 	uint32 NumOriginalWaveformFrames = 0;
 
 	/* The bounds of the transformed waveform in relation to the original */
diff --git a/Engine/Plugins/Editor/WaveformEditor/Source/WaveformEditorWidgets/Public/WaveformTransformationRenderLayerFactory.h b/Engine/Plugins/Editor/WaveformEditor/Source/WaveformEditorWidgets/Public/WaveformTransformationRenderLayerFactory.h
--- a/Engine/Plugins/Editor/WaveformEditor/Source/WaveformEditorWidgets/Public/WaveformTransformationRenderLayerFactory.h
+++ b/Engine/Plugins/Editor/WaveformEditor/Source/WaveformEditorWidgets/Public/WaveformTransformationRenderLayerFactory.h
@@ -17,8 +17,9 @@ public:
 	~FWaveformTransformationRenderLayerFactory() = default;
 
 	UE_API TSharedPtr<IWaveformTransformationRenderer> Create(TObjectPtr<UWaveformTransformationBase> InTransformationToRender);
-	UE_API TSharedPtr<IWaveformTransformationRenderer> Create(TObjectPtr<UWaveformTransformationBase> InTransformationToRender, TArray<TSharedRef<IPropertyHandle>>& TransformationsProperties);
 
+	UE_DEPRECATED(5.7, "TransformationsProperties in Create has been deprecated.")
+	UE_API TSharedPtr<IWaveformTransformationRenderer> Create(TObjectPtr<UWaveformTransformationBase> InTransformationToRender, TArray<TSharedRef<IPropertyHandle>>& TransformationsProperties);
 };
 
 #undef UE_API
diff --git a/Engine/Plugins/Editor/WaveformEditor/Source/WaveformTransformations/Private/WaveformTransformationMarkers.cpp b/Engine/Plugins/Editor/WaveformEditor/Source/WaveformTransformations/Private/WaveformTransformationMarkers.cpp
--- a/Engine/Plugins/Editor/WaveformEditor/Source/WaveformTransformations/Private/WaveformTransformationMarkers.cpp
+++ b/Engine/Plugins/Editor/WaveformEditor/Source/WaveformTransformations/Private/WaveformTransformationMarkers.cpp
@@ -1,11 +1,20 @@
 // Copyright Epic Games, Inc. All Rights Reserved.
 
 #include "WaveformTransformationMarkers.h"
+
+#include "Editor.h"
 #include "Sound/SoundWave.h"
+#include "WaveformTransformationLog.h"
 
 #include UE_INLINE_GENERATED_CPP_BY_NAME(WaveformTransformationMarkers)
 
-DEFINE_LOG_CATEGORY(LogWaveformTransformationMarkers);
+#define LOCTEXT_NAMESPACE "WaveformTransformationMarkers"
+
+namespace
+{
+	static const TArray<int32> SizeIncrements = { 1, 10, 50, 100 };
+	static const TArray<FKey> IndexKeyCodes = { EKeys::One, EKeys::Two, EKeys::Three, EKeys::Four, EKeys::Five, EKeys::Six, EKeys::Seven, EKeys::Eight, EKeys::Nine };
+}
 
 FWaveTransformationMarkers::FWaveTransformationMarkers(double InStartLoopTime, double InEndLoopTime)
 : StartLoopTime(InStartLoopTime)
@@ -25,39 +34,23 @@ void FWaveTransformationMarkers::ProcessAudio(Audio::FWaveformTransformationWave
 	check(StartLoopTime >= 0);
 	check(InOutWaveInfo.NumChannels > 0);
 
-	int64 StartSample = FMath::RoundToInt32(StartLoopTime * InOutWaveInfo.SampleRate) * InOutWaveInfo.NumChannels;
-	StartSample = FMath::Max(0, StartSample);
-
-	if (StartSample > InputAudio.Num())
-	{
-		return;
-	}
+	const int64 LastInputAudioIndex = InputAudio.Num() - 1;
+	const int64 StartSample = FMath::Clamp(FMath::RoundToInt64(StartLoopTime * InOutWaveInfo.SampleRate) * InOutWaveInfo.NumChannels, 0, LastInputAudioIndex - InOutWaveInfo.NumChannels);
+	check(StartSample < LastInputAudioIndex);
 	
-	int64 EndSample = InputAudio.Num() - 1;
+	int64 EndSample = LastInputAudioIndex;
 
 	if (EndLoopTime > 0.f)
 	{
-		const int64 EndFrame = FMath::RoundToInt32(EndLoopTime * InOutWaveInfo.SampleRate) * InOutWaveInfo.NumChannels;
-		EndSample = FMath::Max(StartSample + 1, EndFrame - 1);
+		const int64 EndFrame = FMath::RoundToInt64(EndLoopTime * InOutWaveInfo.SampleRate) * InOutWaveInfo.NumChannels;
+		EndSample = FMath::Max(StartSample + InOutWaveInfo.NumChannels, EndFrame);
 
 		// EndLoopTime can be beyond the length of the file if there is a Trim
-		if (EndSample > InputAudio.Num() - 1)
-		{
-			EndSample = InputAudio.Num() - 1;
-			UE_LOG(LogWaveformTransformationMarkers, Warning, TEXT("Cutting a loop point with a trim!"));
-		}
-
-		EndSample = FMath::Min(EndSample, InputAudio.Num() - 1);
+		EndSample = FMath::Min(EndSample, LastInputAudioIndex);
 	}
 
-	const int32 FinalSize = EndSample - StartSample + 1;
-
-	if (FinalSize <= 2)
-	{
-		UE_LOG(LogWaveformTransformationMarkers, Warning, TEXT("Previewing loop of sample size 1!"));
-		
-		return;
-	}
+	const int32 FinalSize = EndSample - StartSample;
+	check (StartSample + FinalSize <= LastInputAudioIndex)
 
 	InOutWaveInfo.StartFrameOffset = StartSample - (StartSample % InOutWaveInfo.NumChannels);
 	InOutWaveInfo.NumEditedSamples = FinalSize;
@@ -86,6 +79,8 @@ UWaveformTransformationMarkers::UWaveformTransformationMarkers(const FObjectInit
 	{
 		Markers = ObjectInitializer.CreateDefaultSubobject<UWaveCueArray>(this, TEXT("Markers"));
 	}
+
+	Markers->CueChanged.BindLambda([this]() { SetLoopPreviewing(); });
 }
 
 #if WITH_EDITOR
@@ -97,55 +92,19 @@ void UWaveformTransformationMarkers::PostEditChangeProperty(FPropertyChangedEven
 	if (PropertyName == GET_MEMBER_NAME_CHECKED(UWaveformTransformationMarkers, Markers))
 	{
 		check(Markers);
-		bool bIsSelectedCuePresent = false;
-		
-		for (FSoundWaveCuePoint& CuePoint : Markers->CuesAndLoops)
-		{
-			if (!bIsSelectedCuePresent && Markers->SelectedCue == CuePoint.CuePointID)
-			{
-				bIsSelectedCuePresent = true;
-			}
-
-			if(CuePoint.IsLoopRegion())
-			{
-				// Update LoopRegion Preview window if a loop is selected and the user edits the loop bounds using the properties window
-				if (bIsPreviewingLoopRegion && Markers->SelectedCue == CuePoint.CuePointID)
-				{
-					check(SampleRate > 0.f);
-					check(EndLoopTime > 0.0);
-
-					const int64 StartLoopFramePos = (StartLoopTime * SampleRate);
-					const int64 EndLoopFramePos = (EndLoopTime * SampleRate);
-
-					if (CuePoint.FramePosition != StartLoopFramePos)
-					{
-						StartLoopTime = static_cast<float>(FMath::Max(0, CuePoint.FramePosition - StartFrameOffset)) / SampleRate;
-						check(StartLoopTime >= 0);
-					}
-
-					if (CuePoint.FramePosition + CuePoint.FrameLength != EndLoopFramePos)
-					{
-						EndLoopTime = static_cast<float>(CuePoint.FramePosition + CuePoint.FrameLength - StartFrameOffset) / SampleRate;
-					}
-				}
-
-				if (CuePoint.FrameLength < Markers->MinLoopSize)
-				{
-					CuePoint.FrameLength = FMath::Min(Markers->MinLoopSize, static_cast<int64>(AvailableWaveformDuration * SampleRate * 0.1f));
-				}
-			}
-		}
+		SetLoopPreviewing();
+	}
+	else 
+	{
+		OnTransformationChanged.ExecuteIfBound();
+	}
+}
 
-		if (!bIsSelectedCuePresent && Markers->SelectedCue != INDEX_NONE)
-		{
-			Markers->SelectedCue = INDEX_NONE;
-		}
+void UWaveformTransformationMarkers::PostEditUndo()
+{
+	AdjustLoopPreviewIfNotAligned();
 
-		if (Markers->SelectedCue == INDEX_NONE && (bIsPreviewingLoopRegion || StartLoopTime != 0.0 || EndLoopTime != -1.0))
-		{
-			ResetLoopPreviewing();
-		}
-	}
+	Super::PostEditUndo();
 }
 #endif // WITH_EDITOR
 
@@ -169,11 +128,34 @@ void UWaveformTransformationMarkers::UpdateConfiguration(FWaveTransformUObjectCo
 	StartFrameOffset = static_cast<int64>(InOutConfiguration.StartTime * InOutConfiguration.SampleRate);
 	SampleRate = InOutConfiguration.SampleRate;
 	AvailableWaveformDuration = InOutDuration;
+	NumChannels = InOutConfiguration.NumChannels;
+	NumAvailableSamples = static_cast<int64>(InOutDuration * SampleRate);
+
+	if (!bHasLoggedCutByLoop)
+	{
+		for (FSoundWaveCuePoint& CuePoint : Markers->CuesAndLoops)
+		{
+			if (!CuePoint.IsLoopRegion())
+			{
+				continue;
+			}
+
+			const int64 LoopStartSample = CuePoint.FramePosition;
+			const int64 LoopEndSample = CuePoint.FramePosition + CuePoint.FrameLength;
+			const int64 EndSample = FMath::RoundToInt64(InOutConfiguration.EndTime * SampleRate) * NumChannels;
+
+			if ((StartFrameOffset > LoopStartSample && StartFrameOffset < LoopEndSample) || (LoopEndSample > EndSample && LoopStartSample < LoopEndSample))
+			{
+				bHasLoggedCutByLoop = true;
+
+				UE_LOG(LogWaveformTransformation, Log, TEXT("Cutting a loop point with a trim!"));
+			}
+		}
+	}
 
 	if (!bCachedIsPreviewingLoopRegion)
 	{
 		bCachedSoundWaveLoopState = InOutConfiguration.bCachedSoundWaveLoopState;
-		EndLoopTime = -1;
 	}
 	else
 	{
@@ -193,6 +175,7 @@ void UWaveformTransformationMarkers::OverwriteTransformation()
 	Markers->Reset();
 }
 
+PRAGMA_DISABLE_DEPRECATION_WARNINGS
 void UWaveformTransformationMarkers::ModifyMarkerLoopRegion(ELoopModificationControls Modification) const
 {
 	check(Markers != nullptr);
@@ -204,12 +187,291 @@ void UWaveformTransformationMarkers::CycleMarkerLoopRegion(ELoopModificationCont
 	check(Markers != nullptr);
 	Markers->CycleMarkerLoop.ExecuteIfBound(Modification);
 }
+PRAGMA_ENABLE_DEPRECATION_WARNINGS
+
+void UWaveformTransformationMarkers::ModifyMarkerLoopRegion(ELoopModificationControls Modification)
+{
+	check(Markers != nullptr);
+
+	FSoundWaveCuePoint* LoopRegion = GetSelectedMarker();
+
+	if (LoopRegion == nullptr)
+	{
+		return;
+	}
+
+	check(Markers);
+	check(NumChannels > 0);
+	check(NumAvailableSamples > 0);
+	check(SizeIncrements.Num() > 0);
+
+	int64 FramesToShift = static_cast<int64>((SizeIncrements[SelectedIncrement] / 1000.0f) * SampleRate);
+	int64 MaxFrames = static_cast<int64>(NumAvailableSamples) / static_cast<int64>(NumChannels);
+	check(MaxFrames > 0);
+
+	if (GEditor && GEditor->Trans)
+	{
+		GEditor->BeginTransaction(TEXT("PropertyEditor"), LOCTEXT("ModifyingMarkerFrame", "Modifying Marker"), nullptr);
+	}
+
+	Markers->Modify();
+
+	switch (Modification)
+	{
+	case ELoopModificationControls::None:
+		break;
+	case ELoopModificationControls::LeftHandleIncrement:
+	{
+		LoopRegion->FramePosition -= FramesToShift;
+		LoopRegion->FrameLength += FramesToShift;
+	}
+	break;
+	case ELoopModificationControls::LeftHandleDecrement:
+	{
+		LoopRegion->FramePosition += FramesToShift;
+		LoopRegion->FrameLength -= FramesToShift;
+	}
+	break;
+	case ELoopModificationControls::RightHandleIncrement:
+	{
+		LoopRegion->FrameLength += FramesToShift;
+	}
+	break;
+	case ELoopModificationControls::RightHandleDecrement:
+	{
+		LoopRegion->FrameLength -= FramesToShift;
+	}
+	break;
+	case ELoopModificationControls::IncreaseIncrement:
+	{
+		if (SelectedIncrement < SizeIncrements.Num() - 1)
+		{
+			++SelectedIncrement;
+		}
+	}
+	break;
+	case ELoopModificationControls::DecreaseIncrement:
+	{
+		if (SelectedIncrement > 0)
+		{
+			--SelectedIncrement;
+		}
+	}
+	default:
+		break;
+	}
+
+	// Prevent the frame position from exceeding the wave length
+	LoopRegion->FramePosition = FMath::Clamp(LoopRegion->FramePosition, 0, MaxFrames - 1);
+	// FrameLength can exceed MaxFrames because the loop will be cut by ProcessAudio
+	// but length must be >= MinLoopSize or it is no longer a loop
+	LoopRegion->FrameLength = FMath::Max(LoopRegion->FrameLength, Markers->MinLoopSize);
+
+	PreviewSelectedLoop();
+
+	if (GEditor && GEditor->Trans)
+	{
+		GEditor->EndTransaction();
+	}
+
+	OnTransformationChanged.ExecuteIfBound();
+}
+
+void UWaveformTransformationMarkers::CycleMarkerLoopRegion(ELoopModificationControls Modification)
+{
+	check(Markers != nullptr);
+
+	const int32 MarkerArrayLength = Markers->CuesAndLoops.Num();
+
+	FSoundWaveCuePoint* SelectedCuePointPtr = nullptr;
+
+	// If there is no active selection, pick first or last elmeent
+	if (Markers->SelectedCue == INDEX_NONE && MarkerArrayLength > 0)
+	{
+		switch (Modification)
+		{
+		case ELoopModificationControls::SelectNextLoop:
+		{
+			Markers->SelectedCue = Markers->CuesAndLoops[0].CuePointID;
+			SelectedCuePointPtr = &Markers->CuesAndLoops[0];
+		}
+		break;
+		case ELoopModificationControls::SelectPreviousLoop:
+		{
+			Markers->SelectedCue = Markers->CuesAndLoops[MarkerArrayLength - 1].CuePointID;
+			SelectedCuePointPtr = &Markers->CuesAndLoops[MarkerArrayLength - 1];
+		}
+		break;
+		default:
+			break;
+		}
+		return;
+	}
+
+	// Since cues can be removed and added the ids do not match list indexes, so search for matches
+	for (int i = 0; i < MarkerArrayLength; i++)
+	{
+		if (Markers->CuesAndLoops[i].CuePointID == Markers->SelectedCue)
+		{
+			// Check if we move backwards of forwards
+			if (Modification == ELoopModificationControls::SelectNextLoop)
+			{
+				Markers->SelectedCue = Markers->CuesAndLoops[(i + 1) % MarkerArrayLength].CuePointID;
+				SelectedCuePointPtr = &Markers->CuesAndLoops[(i + 1) % MarkerArrayLength];
+				break;
+			}
+			else if (Modification == ELoopModificationControls::SelectPreviousLoop)
+			{
+				Markers->SelectedCue = Markers->CuesAndLoops[(i - 1 + MarkerArrayLength) % MarkerArrayLength].CuePointID;
+				SelectedCuePointPtr = &Markers->CuesAndLoops[(i - 1 + MarkerArrayLength) % MarkerArrayLength];
+				break;
+			}
+		}
+	}
+
+	if (SelectedCuePointPtr && SelectedCuePointPtr->IsLoopRegion())
+	{
+		PreviewSelectedLoop();
+	}
+	else
+	{
+		ResetLoopPreviewing();
+	}
+
+	OnTransformationChanged.ExecuteIfBound();
+}
 
-void UWaveformTransformationMarkers::ResetLoopPreviewing()
+bool UWaveformTransformationMarkers::IsLoopRegionActive() const
 {
-	bIsPreviewingLoopRegion = false;
-	StartLoopTime = 0.0;
-	EndLoopTime = -1.0;
+	if (bIsPreviewingLoopRegion || EndLoopTime != -1.0)
+	{
+		return true;
+	}
+
+	return false;
+}
+
+void UWaveformTransformationMarkers::SetLoopPreviewing()
+{
+	check(Markers);
+	bool bIsSelectedCuePresent = false;
+	bool bTransformationChanged = false;
+
+	for (FSoundWaveCuePoint& CuePoint : Markers->CuesAndLoops)
+	{
+		if (!bIsSelectedCuePresent && Markers->SelectedCue == CuePoint.CuePointID)
+		{
+			bIsSelectedCuePresent = true;
+
+			if (!CuePoint.IsLoopRegion())
+			{
+				bTransformationChanged = ResetLoopPreviewing();
+			}
+		}
+
+		if (CuePoint.IsLoopRegion())
+		{
+			if (CuePoint.FrameLength < Markers->MinLoopSize)
+			{
+				CuePoint.FrameLength = FMath::Max(Markers->MinLoopSize, static_cast<int64>(AvailableWaveformDuration * SampleRate * 0.1f));
+				bTransformationChanged = true;
+			}
+
+			// Update LoopRegion Preview window if a loop is selected and the user edits the loop bounds using the properties window
+			if (IsLoopRegionActive() && Markers->SelectedCue == CuePoint.CuePointID)
+			{
+				check(SampleRate > 0.f);
+				check(EndLoopTime > 0.0);
+
+				const int64 StartLoopFramePos = (StartLoopTime * SampleRate);
+				const int64 EndLoopFramePos = (EndLoopTime * SampleRate);
+
+				if (CuePoint.FramePosition != StartLoopFramePos)
+				{
+					StartLoopTime = static_cast<float>(FMath::Max(0, CuePoint.FramePosition - StartFrameOffset)) / SampleRate;
+					check(StartLoopTime >= 0);
+					bTransformationChanged = true;
+				}
+
+				if (CuePoint.FramePosition + CuePoint.FrameLength != EndLoopFramePos)
+				{
+					EndLoopTime = static_cast<float>(CuePoint.FramePosition + CuePoint.FrameLength) / SampleRate;
+					bTransformationChanged = true;
+				}
+			}
+		}
+	}
+
+	if (!bIsSelectedCuePresent && Markers->SelectedCue != INDEX_NONE)
+	{
+		Markers->SelectedCue = INDEX_NONE;
+	}
+
+	if (Markers->SelectedCue == INDEX_NONE)
+	{
+		bTransformationChanged = ResetLoopPreviewing();
+	}
+
+	if (bTransformationChanged)
+	{
+		OnTransformationChanged.ExecuteIfBound();
+	}
+}
+
+bool UWaveformTransformationMarkers::ResetLoopPreviewing()
+{
+	if (IsLoopRegionActive())
+	{
+		bIsPreviewingLoopRegion = false;
+		StartLoopTime = 0.0;
+		EndLoopTime = -1.0;
+
+		return true;
+	}
+
+	return false;
+}
+
+bool UWaveformTransformationMarkers::AdjustLoopPreviewIfNotAligned()
+{
+	check(Markers);
+	FSoundWaveCuePoint* CuePoint = GetSelectedMarker();
+
+	if (Markers->SelectedCue == INDEX_NONE || !IsLoopRegionActive() || CuePoint == nullptr)
+	{
+		return false;
+	}
+
+	check(SampleRate > 0.f);
+
+	const int64 StartLoopFramePos = static_cast<int64>(StartLoopTime * SampleRate);
+	const int64 EndLoopFramePos = static_cast<int64>(EndLoopTime * SampleRate);
+
+	bool bPreviewChanged = false;
+
+	const int64 CuePointPosDif = FMath::Abs(CuePoint->FramePosition - StartLoopFramePos - StartFrameOffset);
+	const int64 CuePointEndPosDif = FMath::Abs((CuePoint->FramePosition + CuePoint->FrameLength) - EndLoopFramePos - StartFrameOffset);
+
+	if (StartFrameOffset > CuePoint->FramePosition && StartLoopTime == 0.0)
+	{
+		return false;
+	}
+
+	// Converting from time to frame is not perfectly accurate, lets be +- 2 frames
+	if (CuePointPosDif > 2)
+	{
+		StartLoopTime = static_cast<float>(FMath::Max(0, CuePoint->FramePosition - StartFrameOffset)) / SampleRate;
+		bPreviewChanged = true;
+	}
+
+	// Converting from time to frame is not perfectly accurate, lets be +- 2 frames
+	if (CuePointEndPosDif > 2)
+	{
+		EndLoopTime = static_cast<float>((CuePoint->FramePosition + CuePoint->FrameLength) - StartFrameOffset) / SampleRate;
+		bPreviewChanged = true;
+	}
+
+	return bPreviewChanged;
 }
 
 #if WITH_EDITOR
@@ -254,6 +516,91 @@ void UWaveformTransformationMarkers::GetTransformationInfo(FWaveformTransformati
 	InOutTransformationInfo.AllCuePoints.Append(Markers->CuesAndLoops);
 }
 
+FSoundWaveCuePoint* UWaveformTransformationMarkers::GetSelectedMarker() const
+{
+	check(Markers);
+
+	if (Markers->SelectedCue == INDEX_NONE)
+	{
+		return nullptr;
+	}
+
+	for (int i = 0; i < Markers->CuesAndLoops.Num(); i++)
+	{
+		if (Markers->CuesAndLoops[i].CuePointID == Markers->SelectedCue)
+		{
+			return &Markers->CuesAndLoops[i];
+		}
+	}
+
+	return nullptr;
+}
+
+void UWaveformTransformationMarkers::PreviewSelectedLoop()
+{
+	check(SampleRate > 0);
+	check(NumChannels > 0);
+
+	const int64 EndFramePosition = StartFrameOffset + NumAvailableSamples / NumChannels;
+
+	for (FSoundWaveCuePoint CuePoint : Markers->CuesAndLoops)
+	{
+		const int64 LoopEndFramePosition = CuePoint.FramePosition + CuePoint.FrameLength;
+
+		// Only preview a loop if part of it is within the available frames (accounting for TrimFades and other transformations)
+		if (CuePoint.CuePointID == Markers->SelectedCue && CuePoint.IsLoopRegion() &&
+			LoopEndFramePosition >= StartFrameOffset && CuePoint.FramePosition < EndFramePosition)
+		{
+			check(CuePoint.FrameLength > 0);
+
+			// Set Loop preview handles
+			const float StartLoopPosInSeconds = static_cast<float>(CuePoint.FramePosition) / SampleRate;
+			const float EndLoopPosInSeconds = static_cast<float>(LoopEndFramePosition) / SampleRate;
+
+			SetIsPreviewingLoopRegion(static_cast<double>(StartLoopPosInSeconds), static_cast<double>(EndLoopPosInSeconds), true);
+		}
+	}
+}
+
+void UWaveformTransformationMarkers::SetIsPreviewingLoopRegion(double InStartTime, double InEndTime, bool bIsPreviewing)
+{
+	// Stop any current loop previewing so StartLoopTime is never greater than EndLoopTime (Other than when EndLoopTime is invalid)
+	ResetLoopPreviewing();
+
+	bIsPreviewingLoopRegion = bIsPreviewing;
+	StartLoopTime = InStartTime;
+	EndLoopTime = InEndTime;
+}
+
+void UWaveformTransformationMarkers::SelectLoopRegionByKeyboard(const FKey& PressedKey)
+{
+	check(Markers);
+
+	for (int i = 0; i < IndexKeyCodes.Num(); i++)
+	{
+		// Check we have loop regions in this index
+		if (Markers->CuesAndLoops.Num() <= i)
+		{
+			return;
+		}
+
+		if (PressedKey == IndexKeyCodes[i])
+		{
+			Markers->SelectedCue = i;
+
+			check(SampleRate > 0);
+			if (GetSelectedMarker() && GetSelectedMarker()->IsLoopRegion())
+			{
+				PreviewSelectedLoop();
+			}
+			else
+			{
+				SetIsPreviewingLoopRegion(0.0, -1.0, false);
+			}
+		}
+	}
+}
+
 void UWaveCueArray::PostEditChangeProperty(FPropertyChangedEvent& PropertyChangedEvent)
 {
 	Super::PostEditChangeProperty(PropertyChangedEvent);
@@ -327,3 +674,5 @@ void UWaveCueArray::EnableLoopRegion(FSoundWaveCuePoint* OutSoundWaveCue)
 	check(OutSoundWaveCue != nullptr);
 	OutSoundWaveCue->SetLoopRegion(true);
 }
+
+#undef LOCTEXT_NAMESPACE
\ No newline at end of file
diff --git a/Engine/Plugins/Editor/WaveformEditor/Source/WaveformTransformations/Private/WaveformTransformationTrimFade.cpp b/Engine/Plugins/Editor/WaveformEditor/Source/WaveformTransformations/Private/WaveformTransformationTrimFade.cpp
--- a/Engine/Plugins/Editor/WaveformEditor/Source/WaveformTransformations/Private/WaveformTransformationTrimFade.cpp
+++ b/Engine/Plugins/Editor/WaveformEditor/Source/WaveformTransformations/Private/WaveformTransformationTrimFade.cpp
@@ -93,17 +93,23 @@ void FWaveTransformationTrimFade::ProcessAudio(Audio::FWaveformTransformationWav
 		return;
 	}
 
-	const int32 StartSample = FMath::RoundToInt32(StartTime * InOutWaveInfo.SampleRate) * InOutWaveInfo.NumChannels;
-	int32 EndSample = InputAudio.Num() - 1;
+	const int64 LastInputAudioIndex = InputAudio.Num() - 1;
+	int64 StartSample = FMath::Min(FMath::FloorToInt64(StartTime * InOutWaveInfo.SampleRate) * InOutWaveInfo.NumChannels, LastInputAudioIndex - InOutWaveInfo.NumChannels);
+	int64 EndSample = InputAudio.Num();
 	
 	if(EndTime > 0.f)
 	{	
-		const int32 EndFrame = FMath::RoundToInt32(EndTime * InOutWaveInfo.SampleRate);
-		EndSample = EndFrame * InOutWaveInfo.NumChannels - 1;
-		EndSample = FMath::Min(EndSample, InputAudio.Num() - 1);
+		const int64 EndFrame = FMath::RoundToInt64(EndTime * InOutWaveInfo.SampleRate);
+		EndSample = EndFrame * InOutWaveInfo.NumChannels;
+		EndSample = FMath::Min(EndSample, LastInputAudioIndex);
 	}
 
-	const int32 FinalSize = EndSample - StartSample + 1;
+	if (StartSample > EndSample - InOutWaveInfo.NumChannels)
+	{
+		StartSample = EndSample - InOutWaveInfo.NumChannels;
+	}
+
+	const int64 FinalSize = EndSample - StartSample;
 
 	InOutWaveInfo.StartFrameOffset = StartSample - (StartSample % InOutWaveInfo.NumChannels);
 	InOutWaveInfo.NumEditedSamples = FinalSize;
diff --git a/Engine/Plugins/Editor/WaveformEditor/Source/WaveformTransformations/Private/WaveformTransformationsModule.cpp b/Engine/Plugins/Editor/WaveformEditor/Source/WaveformTransformations/Private/WaveformTransformationsModule.cpp
--- a/Engine/Plugins/Editor/WaveformEditor/Source/WaveformTransformations/Private/WaveformTransformationsModule.cpp
+++ b/Engine/Plugins/Editor/WaveformEditor/Source/WaveformTransformations/Private/WaveformTransformationsModule.cpp
@@ -4,8 +4,10 @@
 
 #include "Modules/ModuleManager.h"
 #include "PropertyEditorModule.h"
+#include "WaveformTransformationLog.h"
 #include "WaveformTransformationMarkersObjectCustomization.h"
 
+DEFINE_LOG_CATEGORY(LogWaveformTransformation);
 
 void FWaveformTransformationsModule::StartupModule()
 {
diff --git a/Engine/Plugins/Editor/WaveformEditor/Source/WaveformTransformations/Public/WaveformTransformationLog.h b/Engine/Plugins/Editor/WaveformEditor/Source/WaveformTransformations/Public/WaveformTransformationLog.h
new file mode 100644
--- /dev/null
+++ b/Engine/Plugins/Editor/WaveformEditor/Source/WaveformTransformations/Public/WaveformTransformationLog.h
@@ -0,0 +1,7 @@
+// Copyright Epic Games, Inc. All Rights Reserved.
+
+#pragma once
+
+#include "CoreMinimal.h"
+
+DECLARE_LOG_CATEGORY_EXTERN(LogWaveformTransformation, Log, All);
\ No newline at end of file
diff --git a/Engine/Plugins/Editor/WaveformEditor/Source/WaveformTransformations/Public/WaveformTransformationMarkers.h b/Engine/Plugins/Editor/WaveformEditor/Source/WaveformTransformations/Public/WaveformTransformationMarkers.h
--- a/Engine/Plugins/Editor/WaveformEditor/Source/WaveformTransformations/Public/WaveformTransformationMarkers.h
+++ b/Engine/Plugins/Editor/WaveformEditor/Source/WaveformTransformations/Public/WaveformTransformationMarkers.h
@@ -3,13 +3,12 @@
 #pragma once
 
 #include "IWaveformTransformation.h"
+#include "InputCore.h"
 
 #include "WaveformTransformationMarkers.generated.h"
 
 #define UE_API WAVEFORMTRANSFORMATIONS_API
 
-DECLARE_LOG_CATEGORY_EXTERN(LogWaveformTransformationMarkers, Log, All);
-
 class USoundWave;
 
 enum class ELoopModificationControls : uint8
@@ -44,12 +43,20 @@ public:
 	TArray<FSoundWaveCuePoint> CuesAndLoops;
 	int32 SelectedCue = INDEX_NONE;
 
+	PRAGMA_DISABLE_DEPRECATION_WARNINGS
+
+	UE_DEPRECATED(5.7, "ModifyMarkerLoopRegionDelegate has been deprecated.")
 	DECLARE_DELEGATE_OneParam(ModifyMarkerLoopRegionDelegate, ELoopModificationControls);
+	UE_DEPRECATED(5.7, "ModifyMarkerLoop has been deprecated.")
 	ModifyMarkerLoopRegionDelegate ModifyMarkerLoop;
 
+	UE_DEPRECATED(5.7, "CycleMarkerLoopRegionDelegate has been deprecated.")
 	DECLARE_DELEGATE_OneParam(CycleMarkerLoopRegionDelegate, ELoopModificationControls);
+	UE_DEPRECATED(5.7, "CycleMarkerLoop has been deprecated.")
 	CycleMarkerLoopRegionDelegate CycleMarkerLoop;
 
+	PRAGMA_ENABLE_DEPRECATION_WARNINGS
+
 	// To minimize complexity while supporting all common editing cases, loops have a min length of 10 frames
 	static constexpr int64 MinLoopSize = 10;
 
@@ -83,6 +90,8 @@ public:
 
 #if WITH_EDITOR
 	UE_API virtual void PostEditChangeProperty(FPropertyChangedEvent& PropertyChangedEvent) override;
+	
+	virtual void PostEditUndo() override;
 #endif // WITH_EDITOR
 
 	UPROPERTY(VisibleAnywhere, Category = "Markers")
@@ -103,29 +112,52 @@ public:
 	UE_API void UpdateConfiguration(FWaveTransformUObjectConfiguration& InOutConfiguration) override;
 	UE_API void OverwriteTransformation() override;
 
+	UE_DEPRECATED(5.7, "ModifyMarkerLoopRegion const has been deprecated, use the non-const version")
 	UE_API void ModifyMarkerLoopRegion(ELoopModificationControls Modification) const;
+	UE_DEPRECATED(5.7, "CycleMarkerLoopRegion const has been deprecated, use the non-const version")
 	UE_API void CycleMarkerLoopRegion(ELoopModificationControls Modification) const;
 
-	UE_API void ResetLoopPreviewing();
+	// Shift marker and loop region bounds with toolkit commands bound in the waveform editor
+	UE_API void ModifyMarkerLoopRegion(ELoopModificationControls Modification);
+	// Set SelectedCue by cycling through cues and loops with toolkit commands bound in the waveform editor
+	UE_API void CycleMarkerLoopRegion(ELoopModificationControls Modification);
+
+	UE_API bool IsLoopRegionActive() const;
+	UE_API void SetLoopPreviewing();
+
+	// Returns true if reset
+	UE_API bool ResetLoopPreviewing();
+	// Returns true if reset
+	UE_API bool AdjustLoopPreviewIfNotAligned();
 
 #if WITH_EDITOR
 	UE_API void OverwriteSoundWaveData(USoundWave& InOutSoundWave) override;
 	UE_API void GetTransformationInfo(FWaveformTransformationInfo& InOutTransformationInfo) const override;
 #endif //WITH_EDITOR
 
+	UE_API FSoundWaveCuePoint* GetSelectedMarker() const;
+
 private:
+
+	UE_API void PreviewSelectedLoop();
+	UE_API void SetIsPreviewingLoopRegion(double InStartTime, double InEndTime, bool bIsPreviewing);
+	UE_API void SelectLoopRegionByKeyboard(const FKey& PressedKey);
+	
+	int32 SelectedIncrement = 0;
+
 	float SampleRate = 0.f;
 	float AvailableWaveformDuration = -1.f;
+	int32 NumChannels = 0;
+	int64 NumAvailableSamples = 0;
 
-	/* Made a UPROPERTY so that it is captured when the soundwave is duplicated during export
-	 * Given these specifiers to prevent a crash in the property handle caching system used for transformations
-	 * Note: hidden in the details panel by the IPropertyTypeCustomization
-	 */
-	UPROPERTY(VisibleAnywhere, Category = "Markers")
+	// Made a UPROPERTY so that it is captured when the soundwave is duplicated during export
+	UPROPERTY()
 	int64 StartFrameOffset = 0;
 
 	bool bCachedIsPreviewingLoopRegion = false;
 	bool bCachedSoundWaveLoopState = false;
+
+	bool bHasLoggedCutByLoop = false;
 };
 
 #undef UE_API
diff --git a/Engine/Plugins/Editor/WaveformEditor/Source/WaveformTransformations/WaveformTransformations.build.cs b/Engine/Plugins/Editor/WaveformEditor/Source/WaveformTransformations/WaveformTransformations.build.cs
--- a/Engine/Plugins/Editor/WaveformEditor/Source/WaveformTransformations/WaveformTransformations.build.cs
+++ b/Engine/Plugins/Editor/WaveformEditor/Source/WaveformTransformations/WaveformTransformations.build.cs
@@ -17,8 +17,10 @@ namespace UnrealBuildTool.Rules
 					"Core", 
 					"CoreUObject",
 					"Engine",
+					"InputCore",
 					"PropertyEditor",
-					"SignalProcessing"
+					"SignalProcessing",
+					"UnrealEd"
 				}
 			);
 		}
diff --git a/Engine/Plugins/Editor/WaveformEditor/Source/WaveformTransformationsWidgets/Private/WaveformTransformationMarkersRenderer.cpp b/Engine/Plugins/Editor/WaveformEditor/Source/WaveformTransformationsWidgets/Private/WaveformTransformationMarkersRenderer.cpp
--- a/Engine/Plugins/Editor/WaveformEditor/Source/WaveformTransformationsWidgets/Private/WaveformTransformationMarkersRenderer.cpp
+++ b/Engine/Plugins/Editor/WaveformEditor/Source/WaveformTransformationsWidgets/Private/WaveformTransformationMarkersRenderer.cpp
@@ -3,6 +3,7 @@
 #include "WaveformTransformationMarkersRenderer.h"
 
 #include "Brushes/SlateRoundedBoxBrush.h"
+#include "Editor.h"
 #include "Framework/Application/SlateApplication.h"
 #include "IWaveformTransformation.h"
 #include "Layout/Geometry.h"
@@ -98,6 +99,7 @@ int32 FWaveformTransformationMarkerRenderer::DrawMarkerHandles(const FGeometry&
 
 		float colorMultiplier = 0.5;
 
+		check(MarkersArray.IsValid());
 		// Highlight selection
 		if (MarkersArray->SelectedCue != INDEX_NONE && MarkersArray->SelectedCue == LoopRegion.CuePointID)
 		{
@@ -215,212 +217,12 @@ int32 FWaveformTransformationMarkerRenderer::DrawMarkerHandles(const FGeometry&
 	return LayerId;
 }
 
-void FWaveformTransformationMarkerRenderer::SetPropertyHandles(const TArray<TSharedRef<IPropertyHandle>>& InPropertyHandles)
+void FWaveformTransformationMarkerRenderer::SetWaveformTransformation(TObjectPtr<UWaveformTransformationBase> InTransformation)
 {
-	FWaveformTransformationRendererBase::SetPropertyHandles(InPropertyHandles);
+	FWaveformTransformationRendererBase::SetWaveformTransformation(InTransformation);
+	StrongMarkersTransformation = TStrongObjectPtr<UWaveformTransformationMarkers>(CastChecked<UWaveformTransformationMarkers>(InTransformation.Get()));
 
-	MarkersHandle = GetPropertyHandle(GET_MEMBER_NAME_CHECKED(UWaveformTransformationMarkers, Markers));
-	StartLoopTimeHandle = GetPropertyHandle(GET_MEMBER_NAME_CHECKED(UWaveformTransformationMarkers, StartLoopTime));
-	EndLoopTimeHandle = GetPropertyHandle(GET_MEMBER_NAME_CHECKED(UWaveformTransformationMarkers, EndLoopTime));
-	IsPreviewingLoopRegionHandle = GetPropertyHandle(GET_MEMBER_NAME_CHECKED(UWaveformTransformationMarkers, bIsPreviewingLoopRegion));
-
-	MarkersArray = TStrongObjectPtr<UWaveCueArray>(CastChecked<UWaveCueArray>(GetPropertyValue<UObject*>(MarkersHandle)));
-
-	MarkersArray->ModifyMarkerLoop.BindLambda([this](ELoopModificationControls Modification) { ModifyMarkerLoopRegion(Modification); });
-	MarkersArray->CycleMarkerLoop.BindLambda([this](ELoopModificationControls Modification) { CycleMarkerLoopRegion(Modification); });
-
-	// Flag that the loop preview needs to refresh next Tick to avoid issues starting transactions during an Undo
-	MarkersArray->CueChanged.BindLambda([this]() { bLoopRegionPreviewNeedsRefresh = true; });
-}
-
-void FWaveformTransformationMarkerRenderer::ModifyMarkerLoopRegion(ELoopModificationControls Modification)
-{
-	FSoundWaveCuePoint* LoopRegion = GetSelectedMarker();
-
-	if (LoopRegion == nullptr)
-	{
-		return;
-	}
-
-	check(MarkersArray);
-	check(TransformationWaveInfo.NumChannels > 0);
-	check(TransformationWaveInfo.TotalNumSamples > 0);
-	check(SizeIncrements.Num() > 0);
-
-	int64 FramesToShift = static_cast<int64>((SizeIncrements[SelectedIncrement] / 1000.0f) * TransformationWaveInfo.SampleRate);
-	int64 MaxFrames = static_cast<int64>(TransformationWaveInfo.TotalNumSamples) / static_cast<int64>(TransformationWaveInfo.NumChannels);
-	check(MaxFrames > 0);
-	
-	BeginTransaction(TEXT("PropertyEditor"), LOCTEXT("SetNumericPropertyTransaction", "Modifying Marker"), nullptr);
-
-	MarkersArray->Modify();
-
-	switch (Modification)
-	{
-	case ELoopModificationControls::None:
-		break;
-	case ELoopModificationControls::LeftHandleIncrement:
-	{
-		LoopRegion->FramePosition -= FramesToShift;
-		LoopRegion->FrameLength += FramesToShift;
-	}
-		break;
-	case ELoopModificationControls::LeftHandleDecrement:
-	{
-		LoopRegion->FramePosition += FramesToShift;
-		LoopRegion->FrameLength -= FramesToShift;
-	}
-		break;
-	case ELoopModificationControls::RightHandleIncrement:
-	{
-		LoopRegion->FrameLength += FramesToShift;
-	}
-		break;
-	case ELoopModificationControls::RightHandleDecrement:
-	{
-		LoopRegion->FrameLength -= FramesToShift;
-	}
-		break;
-	case ELoopModificationControls::IncreaseIncrement:
-	{
-		if (SelectedIncrement < SizeIncrements.Num() - 1)
-		{
-			++SelectedIncrement;
-		}
-	}
-	break;
-	case ELoopModificationControls::DecreaseIncrement:
-	{
-		if (SelectedIncrement > 0)
-		{
-			--SelectedIncrement;
-		}
-	}
-	default:
-		break;
-	}
-
-	// Prevent the frame position from exceeding the wave length
-	LoopRegion->FramePosition = FMath::Clamp(LoopRegion->FramePosition, 0, MaxFrames - 1);
-	// FrameLength can exceed MaxFrames because the loop will be cut by ProcessAudio
-	// but length must be > 0 or it is no longer a loop
-	LoopRegion->FrameLength = FMath::Max(LoopRegion->FrameLength, 1);
-
-	PreviewSelectedLoop();
-
-	EndTransaction();
-}
-
-void FWaveformTransformationMarkerRenderer::CycleMarkerLoopRegion(ELoopModificationControls Modification)
-{
-	check(MarkersArray);
-
-	const int32 MarkerArrayLength = MarkersArray->CuesAndLoops.Num();
-
-	FSoundWaveCuePoint* SelectedCuePointPtr = nullptr;
-
-	// If there is no active selection, pick first or last elmeent
-	if (MarkersArray->SelectedCue == INDEX_NONE && MarkerArrayLength > 0)
-	{
-		switch (Modification)
-		{
-		case ELoopModificationControls::SelectNextLoop:
-		{
-			MarkersArray->SelectedCue = MarkersArray->CuesAndLoops[0].CuePointID;
-			SelectedCuePointPtr = &MarkersArray->CuesAndLoops[0];
-		}
-			break;
-		case ELoopModificationControls::SelectPreviousLoop:
-		{
-			MarkersArray->SelectedCue = MarkersArray->CuesAndLoops[MarkerArrayLength - 1].CuePointID;
-			SelectedCuePointPtr = &MarkersArray->CuesAndLoops[MarkerArrayLength - 1];
-		}
-			break;
-		default:
-			break;
-		}
-		return;
-	}
-
-	// Since cues can be removed and added the ids do not match list indexes, so search for matches
-	for (int i = 0; i < MarkerArrayLength; i++)
-	{
-		if (MarkersArray->CuesAndLoops[i].CuePointID == MarkersArray->SelectedCue)
-		{
-			// Check if we move backwards of forwards
-			if (Modification == ELoopModificationControls::SelectNextLoop)
-			{
-				MarkersArray->SelectedCue = MarkersArray->CuesAndLoops[(i + 1) % MarkerArrayLength].CuePointID;
-				SelectedCuePointPtr = &MarkersArray->CuesAndLoops[(i + 1) % MarkerArrayLength];
-				break;
-			}
-			else if (Modification == ELoopModificationControls::SelectPreviousLoop)
-			{
-				MarkersArray->SelectedCue = MarkersArray->CuesAndLoops[(i - 1 + MarkerArrayLength) % MarkerArrayLength].CuePointID;
-				SelectedCuePointPtr = &MarkersArray->CuesAndLoops[(i - 1 + MarkerArrayLength) % MarkerArrayLength];
-				break;
-			}
-		}
-	}
-
-	check(TransformationWaveInfo.SampleRate > 0);
-	if (SelectedCuePointPtr && SelectedCuePointPtr->IsLoopRegion())
-	{
-		PreviewSelectedLoop();
-	}
-	else
-	{
-		SetIsPreviewingLoopRegion(0.0, -1.0, false);
-	}
-}
-
-FSoundWaveCuePoint* FWaveformTransformationMarkerRenderer::GetSelectedMarker() const
-{
-	check(MarkersArray);
-
-	if (MarkersArray->SelectedCue == INDEX_NONE)
-	{
-		return nullptr;
-	}
-
-	for (int i = 0; i < MarkersArray->CuesAndLoops.Num(); i++)
-	{
-		if (MarkersArray->CuesAndLoops[i].CuePointID == MarkersArray->SelectedCue)
-		{
-			return &MarkersArray->CuesAndLoops[i];
-		}
-	}
-
-	return nullptr;
-}
-
-void FWaveformTransformationMarkerRenderer::SelectLoopRegionByKeyboard(const FKey& PressedKey)
-{
-	check(MarkersArray);
-
-	for (int i = 0; i < IndexKeyCodes.Num(); i++)
-	{
-		// Check we have loop regions in this index
-		if (MarkersArray->CuesAndLoops.Num() <= i)
-		{
-			return;
-		}
-
-		if (PressedKey == IndexKeyCodes[i])
-		{
-			MarkersArray->SelectedCue = i;
-
-			check(TransformationWaveInfo.SampleRate > 0);
-			if (GetSelectedMarker() && GetSelectedMarker()->IsLoopRegion())
-			{
-				PreviewSelectedLoop();
-			}
-			else
-			{
-				SetIsPreviewingLoopRegion(0.0, -1.0, false);
-			}
-		}
-	}
+	MarkersArray = TStrongObjectPtr<UWaveCueArray>(CastChecked<UWaveCueArray>(StrongMarkersTransformation->Markers));
 }
 
 void FWaveformTransformationMarkerRenderer::Tick(const FGeometry& AllottedGeometry, const double InCurrentTime, const float InDeltaTime)
@@ -435,12 +237,14 @@ void FWaveformTransformationMarkerRenderer::Tick(const FGeometry& AllottedGeomet
 	FVector2D MouseAbsolutePosition = UWidgetLayoutLibrary::GetMousePositionOnPlatform();
 	MousePosition = AllottedGeometry.AbsoluteToLocal(MouseAbsolutePosition);
 
+	check(StrongMarkersTransformation);
 	check(MarkersArray != nullptr);
 
 	TMap<int32, FSoundWaveCuePoint> CuePointMap;
 	CuePointMap.Reserve(MarkersArray->CuesAndLoops.Num());
 
-	bool IsMarkerInInteractionRangeValid = false;
+	bool bIsMarkerInInteractionRangeValid = false;
+	bool bFixedDuplicateCuePoints = false;
 	for (const FSoundWaveCuePoint& Marker : MarkersArray->CuesAndLoops)
 	{
 		//Fix duplicate loop region markers for SoundWaves that have not been reimported 
@@ -453,11 +257,13 @@ void FWaveformTransformationMarkerRenderer::Tick(const FGeometry& AllottedGeomet
 		else if (Marker.FrameLength > 0)
 		{
 			CuePointMap[Marker.CuePointID].FrameLength = Marker.FrameLength;
+			CuePointMap[Marker.CuePointID].SetLoopRegion(true);
+			bFixedDuplicateCuePoints = true;
 		}
 
 		if (&Marker == MarkerInInteractionRange.GetMarker())
 		{
-			IsMarkerInInteractionRangeValid = true;
+			bIsMarkerInInteractionRangeValid = true;
 		}
 
 		check(TransformationWaveInfo.NumChannels > 0);
@@ -468,21 +274,38 @@ void FWaveformTransformationMarkerRenderer::Tick(const FGeometry& AllottedGeomet
 		if (Marker.CuePointID == MarkersArray->SelectedCue && Marker.IsLoopRegion() &&
 			(Marker.FramePosition + Marker.FrameLength < StartFrameOffset || Marker.FramePosition > StartFrameOffset + NumAvailableSamples))
 		{
-			MarkersHandle->NotifyPreChange();
 			MarkersArray->Modify();
 			MarkersArray->SelectedCue = INDEX_NONE;
 			
-			SetIsPreviewingLoopRegion(0.0, -1.0, false);
+			if (StrongMarkersTransformation->ResetLoopPreviewing())
+			{
+				StrongMarkersTransformation->OnTransformationChanged.ExecuteIfBound();
+			}
+
+			//End move transaction
 			EndTransaction();
 		}
 	}
 
+	if (StrongMarkersTransformation->AdjustLoopPreviewIfNotAligned())
+	{
+		StrongMarkersTransformation->OnTransformationChanged.ExecuteIfBound();
+	}
+
 	//Invalidate Marker pointer if it has been removed from the array
-	if (!IsMarkerInInteractionRangeValid)
+	if (!bIsMarkerInInteractionRangeValid)
 	{
 		MarkerInInteractionRange.SetMarkerInInteractionRange(nullptr);
 	}
 
+	if (bFixedDuplicateCuePoints && !CuePointMap.IsEmpty())
+	{
+		BeginTransaction(TEXT("PropertyEditor"), LOCTEXT("RectifyingMarkerArray", "Rectify Marker Array (Fixed duplicate values)"), StrongMarkersTransformation.Get());
+		MarkersArray->Modify();
+		CuePointMap.GenerateValueArray(MarkersArray->CuesAndLoops);
+		EndTransaction();
+	}
+
 	CuePoints.Empty();
 	LoopPoints.Empty();
 	for (TPair<int32, FSoundWaveCuePoint>& Pair : CuePointMap)
@@ -496,14 +319,8 @@ void FWaveformTransformationMarkerRenderer::Tick(const FGeometry& AllottedGeomet
 			CuePoints.Add(MoveTemp(Pair.Value));
 		}
 	}
-	CuePointMap.Empty();
-
-	if (bLoopRegionPreviewNeedsRefresh)
-	{
-		SetActiveLoopRegion(MarkersArray->SelectedCue);
 
-		bLoopRegionPreviewNeedsRefresh = false;
-	}
+	CuePointMap.Empty();
 }
 
 FReply FWaveformTransformationMarkerRenderer::OnMouseButtonDown(SWidget& OwnerWidget, const FGeometry& MyGeometry, const FPointerEvent& MouseEvent)
@@ -514,7 +331,7 @@ FReply FWaveformTransformationMarkerRenderer::OnMouseButtonDown(SWidget& OwnerWi
 
 	if (MarkerInteractionType != EMarkerInteractionType::None)
 	{
-		MarkersHandle->NotifyPreChange();
+		MarkersArray->Modify();
 
 		return FReply::Handled().CaptureMouse(OwnerWidget.AsShared()).PreventThrottling();
 	}
@@ -524,13 +341,17 @@ FReply FWaveformTransformationMarkerRenderer::OnMouseButtonDown(SWidget& OwnerWi
 
 FReply FWaveformTransformationMarkerRenderer::OnMouseButtonUp(SWidget& OwnerWidget, const FGeometry& MyGeometry, const FPointerEvent& MouseEvent)
 {
+	check(StrongMarkersTransformation);
+	bool bHasTransformationChanged = false;
+
 	if (MarkersArray->SelectedCue == INDEX_NONE)
 	{
-		SetIsPreviewingLoopRegion(0.0, -1.0, false);
+		bHasTransformationChanged = StrongMarkersTransformation->ResetLoopPreviewing();
 	}
 	else
 	{
 		PreviewSelectedLoop();
+		bHasTransformationChanged = true;
 	}
 
 	if (MarkersArray == nullptr)
@@ -540,11 +361,9 @@ FReply FWaveformTransformationMarkerRenderer::OnMouseButtonUp(SWidget& OwnerWidg
 
 	if (MarkerInteractionType != EMarkerInteractionType::None)
 	{
-		MarkersArray->Modify();
+		StrongMarkersTransformation->Modify();
 		SetPropertyValueDependingOnInteractionType(MouseEvent, MyGeometry);
 		
-		MarkersHandle->NotifyPostChange(EPropertyChangeType::ValueSet);
-
 		// If we know it was an interaction that started a transaction, end the transaction
 		if (MarkerInteractionType == EMarkerInteractionType::ScrubbingMarkerHandleLeft || MarkerInteractionType == EMarkerInteractionType::ScrubbingMarkerHandleRight
 			|| MarkerInteractionType == EMarkerInteractionType::LoopHandle || MarkerInteractionType == EMarkerInteractionType::MarkerHandle)
@@ -554,17 +373,26 @@ FReply FWaveformTransformationMarkerRenderer::OnMouseButtonUp(SWidget& OwnerWidg
 
 		MarkerInteractionType = EMarkerInteractionType::None;
 		
+		StrongMarkersTransformation->OnTransformationChanged.ExecuteIfBound();
+		
 		return FReply::Handled().ReleaseMouseCapture();
 	}
 
+	if (bHasTransformationChanged)
+	{
+		StrongMarkersTransformation->OnTransformationChanged.ExecuteIfBound();
+	}
+
 	return FReply::Unhandled();
 }
 
 FReply FWaveformTransformationMarkerRenderer::OnMouseMove(SWidget& OwnerWidget, const FGeometry& MyGeometry, const FPointerEvent& MouseEvent)
 {
+	check(StrongMarkersTransformation);
 	if (MouseEvent.IsMouseButtonDown(EKeys::LeftMouseButton) && MarkerInteractionType != EMarkerInteractionType::None)
 	{
 		SetPropertyValueDependingOnInteractionType(MouseEvent, MyGeometry);
+		StrongMarkersTransformation->OnTransformationRenderChanged.ExecuteIfBound();
 
 		return FReply::Handled().CaptureMouse(OwnerWidget.AsShared());
 	}
@@ -632,7 +460,6 @@ FCursorReply FWaveformTransformationMarkerRenderer::OnCursorQuery(const FGeometr
 
 FReply FWaveformTransformationMarkerRenderer::OnKeyUp(const FGeometry& MyGeometry, const FKeyEvent& InKeyEvent)
 {
-	SelectLoopRegionByKeyboard(InKeyEvent.GetKey());
 	return FReply::Handled();
 }
 
@@ -807,13 +634,14 @@ bool FWaveformTransformationMarkerRenderer::SelectLoopArea(const FVector2D& InLo
 		}
 	}
 
-	SetActiveLoopRegion(INDEX_NONE);
+	MarkersArray->SelectedCue = INDEX_NONE;
 
 	return false;
 }
 
 void FWaveformTransformationMarkerRenderer::SetActiveLoopRegion(int32 CuePointID)
 {
+	check(StrongMarkersTransformation.IsValid());
 	check(MarkersArray != nullptr);
 
 	MarkersArray->Modify();
@@ -821,7 +649,7 @@ void FWaveformTransformationMarkerRenderer::SetActiveLoopRegion(int32 CuePointID
 
 	check(TransformationWaveInfo.SampleRate > 0);
 	// Revert Loop preview handles
-	if (CuePointID == INDEX_NONE || GetSelectedMarker() == nullptr || !GetSelectedMarker()->IsLoopRegion())
+	if (CuePointID == INDEX_NONE || StrongMarkersTransformation->GetSelectedMarker() == nullptr || !StrongMarkersTransformation->GetSelectedMarker()->IsLoopRegion())
 	{
 		SetIsPreviewingLoopRegion(0.0, -1.0, false);
 	}
@@ -864,14 +692,13 @@ void FWaveformTransformationMarkerRenderer::PreviewSelectedLoop()
 
 void FWaveformTransformationMarkerRenderer::SetIsPreviewingLoopRegion(double InStartTime, double InEndTime, bool bIsPreviewing)
 {
+	check(StrongMarkersTransformation);
 	// Stop any current loop previewing so StartLoopTime is never greater than EndLoopTime (Other than when EndLoopTime is invalid)
-	IsPreviewingLoopRegionHandle->SetValue(false);
-	StartLoopTimeHandle->SetValue(0.0);
-	EndLoopTimeHandle->SetValue(-1.0);
+	StrongMarkersTransformation->ResetLoopPreviewing();
 
-	IsPreviewingLoopRegionHandle->SetValue(bIsPreviewing);
-	StartLoopTimeHandle->SetValue(InStartTime);
-	EndLoopTimeHandle->SetValue(InEndTime);
+	StrongMarkersTransformation->bIsPreviewingLoopRegion = bIsPreviewing;
+	StrongMarkersTransformation->StartLoopTime = InStartTime;
+	StrongMarkersTransformation->EndLoopTime = InEndTime;
 }
 
 FVector2D FWaveformTransformationMarkerRenderer::GetLocalCursorPosition(const FPointerEvent& MouseEvent, const FGeometry& EventGeometry)
@@ -883,6 +710,7 @@ FVector2D FWaveformTransformationMarkerRenderer::GetLocalCursorPosition(const FP
 
 FWaveformTransformationMarkerRenderer::EMarkerInteractionType FWaveformTransformationMarkerRenderer::EvaluateInteractionTypeFromCursorPosition(const FVector2D& InLocalCursorPosition, const FGeometry& WidgetGeometry, const FKey MouseButton)
 {
+	check(StrongMarkersTransformation);
 	if (MouseButton == EKeys::RightMouseButton)
 	{
 		if (SetMarkerInInteractionRange(InLocalCursorPosition, WidgetGeometry) && MarkerInInteractionRange.GetMarker() != nullptr && MarkerInInteractionRange.GetMarker()->IsLoopRegion() == false)
@@ -911,7 +739,7 @@ FWaveformTransformationMarkerRenderer::EMarkerInteractionType FWaveformTransform
 		if (IsPositionInLoopHandleArea(InLocalCursorPosition, WidgetGeometry, CuePoint, true))
 		{
 			SetMarkerInInteractionRange(InLocalCursorPosition, WidgetGeometry);
-			BeginTransaction(TEXT("PropertyEditor"), LOCTEXT("ResizingMarker", "ResizingMarker"), nullptr);
+			BeginTransaction(TEXT("PropertyEditor"), LOCTEXT("ResizingMarker", "Resizing Marker"), StrongMarkersTransformation.Get());
 			MarkersArray->Modify();
 			return EMarkerInteractionType::ScrubbingMarkerHandleLeft;
 		}
@@ -919,7 +747,7 @@ FWaveformTransformationMarkerRenderer::EMarkerInteractionType FWaveformTransform
 		if (IsPositionInLoopHandleArea(InLocalCursorPosition, WidgetGeometry, CuePoint, false))
 		{
 			SetMarkerInInteractionRange(InLocalCursorPosition, WidgetGeometry);
-			BeginTransaction(TEXT("PropertyEditor"), LOCTEXT("ResizingMarker", "ResizingMarker"), nullptr);
+			BeginTransaction(TEXT("PropertyEditor"), LOCTEXT("ResizingMarker", "Resizing Marker"), StrongMarkersTransformation.Get());
 			MarkersArray->Modify();
 			return EMarkerInteractionType::ScrubbingMarkerHandleRight;
 		}
@@ -927,7 +755,7 @@ FWaveformTransformationMarkerRenderer::EMarkerInteractionType FWaveformTransform
 
 	if (SetMarkerInInteractionRange(InLocalCursorPosition, WidgetGeometry))
 	{
-		BeginTransaction(TEXT("PropertyEditor"), LOCTEXT("MoveMarker", "MoveMarker"), nullptr);
+		BeginTransaction(TEXT("PropertyEditor"), LOCTEXT("MoveMarker", "Move Marker"), StrongMarkersTransformation.Get());
 		MarkersArray->Modify();
 		return EMarkerInteractionType::MarkerHandle;
 	}
@@ -935,10 +763,17 @@ FWaveformTransformationMarkerRenderer::EMarkerInteractionType FWaveformTransform
 	if (SelectLoopArea(InLocalCursorPosition, WidgetGeometry))
 	{
 		SetMarkerInInteractionRange(InLocalCursorPosition, WidgetGeometry);
-		BeginTransaction(TEXT("PropertyEditor"), LOCTEXT("MoveLoop", "MoveLoop"), nullptr);
+		BeginTransaction(TEXT("PropertyEditor"), LOCTEXT("MoveLoop", "Move Loop"), StrongMarkersTransformation.Get());
 		MarkersArray->Modify();
 		return EMarkerInteractionType::LoopHandle;
 	}
+	else
+	{
+		if (StrongMarkersTransformation->IsLoopRegionActive())
+		{
+			return EMarkerInteractionType::DeselectRegion;
+		}
+	}
 
 	return EMarkerInteractionType::None;
 }
@@ -946,11 +781,9 @@ FWaveformTransformationMarkerRenderer::EMarkerInteractionType FWaveformTransform
 void FWaveformTransformationMarkerRenderer::SetPropertyValueDependingOnInteractionType(const FPointerEvent& MouseEvent, const FGeometry& WidgetGeometry)
 {
 	ensure(PixelsPerFrame > 0);
-	check(MarkersHandle);
 
 	const FVector2D LocalCursorPosition = GetLocalCursorPosition(MouseEvent, WidgetGeometry);
 	const float LocalCursorXClamped = FMath::Clamp(LocalCursorPosition.X, 0.f, WidgetGeometry.GetLocalSize().X);
-	MarkersHandle->NotifyPreChange();
 
 	switch (MarkerInteractionType)
 	{
@@ -1004,58 +837,72 @@ void FWaveformTransformationMarkerRenderer::SetPropertyValueDependingOnInteracti
 void FWaveformTransformationMarkerRenderer::ShowConvertToLoopAtMarker(const FGeometry& MyGeometry, const FPointerEvent& MouseEvent, FSoundWaveCuePoint* CueToModify)
 {
 	check(CueToModify);
+	check(StrongMarkersTransformation);
 	check(MarkersArray != nullptr);
 	const FVector2D LocalWindowMaxPosition = MyGeometry.GetAbsolutePosition() + MyGeometry.GetAbsoluteSize();
 	const FVector2D LocalCursorPosition = MouseEvent.GetScreenSpacePosition();
 
 	TSharedRef<SVerticalBox> MenuContent = SNew(SVerticalBox);
 
-	if (TSharedPtr<SWindow> LockedFadeModeMenuWindow = CreateLoopMenuWindow.Pin())
+	if (TSharedPtr<SWindow> LockedLoopMenuWindow = CreateLoopMenuWindow.Pin())
 	{
-		FSlateApplication::Get().RequestDestroyWindow(LockedFadeModeMenuWindow.ToSharedRef());
+		FSlateApplication::Get().RequestDestroyWindow(LockedLoopMenuWindow.ToSharedRef());
+		FSlateApplication::Get().OnFocusChanging().Remove(PopupHandle);
+		FSlateApplication::Get().OnApplicationActivationStateChanged().Remove(ApplicationActivationStateHandle);
 	}
 
+	TSharedRef<SWindow> MenuWindow = SNew(SWindow)
+		.AutoCenter(EAutoCenter::None)
+		.ScreenPosition(LocalCursorPosition)
+		.SizingRule(ESizingRule::FixedSize)
+		.SupportsTransparency(EWindowTransparency::None)
+		.IsPopupWindow(true)
+		.CreateTitleBar(false)
+		[
+			MenuContent
+		];
+
+	CreateLoopMenuWindow = MenuWindow;
+	TWeakObjectPtr<UWaveformTransformationMarkers> WeakMarkersTransformation = StrongMarkersTransformation.Get();
+
 	MenuContent->AddSlot()
 		.AutoHeight()
 		[
 			SNew(SButton)
-				.Text(LOCTEXT("WaveformTransformationMarkerRendererConvertLoop","Convert To Loop Region"))
-				.OnClicked_Lambda([this, CueToModify]() -> FReply
+				.Text(LOCTEXT("WaveformTransformationMarkerRendererConvertLoop", "Convert To Loop Region"))
+				.OnClicked_Lambda([InCreateLoopMenuWindow = CreateLoopMenuWindow, WeakMarkersTransformation, CueToModify, MinLoopSize = MarkerInInteractionRange.GetMinLoopSize()]() -> FReply
 					{
-						BeginTransaction(TEXT("PropertyEditor"), LOCTEXT("ConvertMarkerToLoop", "ConvertMarkerToLoop"), nullptr);
-
-						MarkersArray->Modify();
-						MarkersArray->EnableLoopRegion(CueToModify);
-						CueToModify->FrameLength = MarkerInInteractionRange.GetMinLoopSize();
+						if (TStrongObjectPtr<UWaveformTransformationMarkers> LockedMarkers = WeakMarkersTransformation.Pin())
+						{
+							if (GEditor && GEditor->Trans)
+							{
+								GEditor->BeginTransaction(TEXT("PropertyEditor"), LOCTEXT("ConvertMarkerToLoop", "Convert Marker To Loop"), LockedMarkers.Get());
+							}
+
+							LockedMarkers->Markers->Modify();
+							LockedMarkers->Markers->EnableLoopRegion(CueToModify);
+							CueToModify->FrameLength = MinLoopSize;
+						}
 
-						EndTransaction();
+						if (GEditor && GEditor->Trans)
+						{
+							GEditor->EndTransaction();
+						}
 
-						if (TSharedPtr<SWindow> LockedFadeModeMenuWindow = CreateLoopMenuWindow.Pin())
+						if (TSharedPtr<SWindow> LockedFadeModeMenuWindow = InCreateLoopMenuWindow.Pin())
 						{
 							FSlateApplication::Get().RequestDestroyWindow(LockedFadeModeMenuWindow.ToSharedRef());
 						}
 						return FReply::Handled();
 					})
-		];
 
-	TSharedRef<SWindow> MenuWindow = SNew(SWindow)
-		.AutoCenter(EAutoCenter::None)
-		.ScreenPosition(LocalCursorPosition)
-		.SizingRule(ESizingRule::Autosized)
-		.SupportsTransparency(EWindowTransparency::None)
-		.IsPopupWindow(true)
-		.CreateTitleBar(false)
-		[
-			MenuContent
 		];
-	
-	CreateLoopMenuWindow = MenuWindow;
 
 	FSlateApplication::Get().AddWindow(MenuWindow);
 
-	FTSTicker::GetCoreTicker().AddTicker(FTickerDelegate::CreateLambda([this, LocalCursorPosition, LocalWindowMaxPosition](float DeltaTime)
+	FTSTicker::GetCoreTicker().AddTicker(FTickerDelegate::CreateLambda([this, InCreateLoopMenuWindow = CreateLoopMenuWindow, LocalCursorPosition, LocalWindowMaxPosition](float DeltaTime)
 		{
-			if (TSharedPtr<SWindow> LockedLoopMenuWindow = CreateLoopMenuWindow.Pin())
+			if (TSharedPtr<SWindow> LockedLoopMenuWindow = InCreateLoopMenuWindow.Pin())
 			{
 				const FVector2D WindowSize = LockedLoopMenuWindow->GetSizeInScreen();
 
@@ -1066,6 +913,11 @@ void FWaveformTransformationMarkerRenderer::ShowConvertToLoopAtMarker(const FGeo
 					const double XPosition = FMath::Max(0.0, LocalCursorPosition.X - WindowSize.X);
 					LockedLoopMenuWindow->MoveWindowTo(FVector2D(XPosition, LocalCursorPosition.Y));
 				}
+
+				const FVector2D SizeScalar(1.1f, 0.8f); // Computed size is too skinny and too tall so we apply this scalar
+				FVector2D ComputedSize = LockedLoopMenuWindow->ComputeWindowSizeForContent(LockedLoopMenuWindow->GetContent()->GetDesiredSize());
+				ComputedSize *= SizeScalar;
+				LockedLoopMenuWindow->Resize(ComputedSize);
 			}
 
 			return false; // return false to only execute once
@@ -1074,17 +926,17 @@ void FWaveformTransformationMarkerRenderer::ShowConvertToLoopAtMarker(const FGeo
 	FSlateApplication::Get().OnFocusChanging().Remove(PopupHandle);
 	FSlateApplication::Get().OnApplicationActivationStateChanged().Remove(ApplicationActivationStateHandle);
 
-	ApplicationActivationStateHandle = FSlateApplication::Get().OnApplicationActivationStateChanged().AddLambda([this](bool isActive)
-		{
-			if (!isActive)
+	if (TSharedPtr<SWindow> LockedLoopMenuWindow = CreateLoopMenuWindow.Pin())
+		ApplicationActivationStateHandle = FSlateApplication::Get().OnApplicationActivationStateChanged().AddLambda([InCreateLoopMenuWindow = CreateLoopMenuWindow](bool isActive)
 			{
-				if (TSharedPtr<SWindow> LockedLoopMenuWindow = CreateLoopMenuWindow.Pin())
+				if (!isActive)
 				{
-					FSlateApplication::Get().RequestDestroyWindow(LockedLoopMenuWindow.ToSharedRef());
-					FSlateApplication::Get().OnApplicationActivationStateChanged().Remove(ApplicationActivationStateHandle);
+					if (TSharedPtr<SWindow> LockedLoopMenuWindow = InCreateLoopMenuWindow.Pin())
+					{
+						FSlateApplication::Get().RequestDestroyWindow(LockedLoopMenuWindow.ToSharedRef());
+					}
 				}
-			}
-		});
+			});
 
 	// If focus is lost on the popup, destroy it to prevent popups hanging around
 	PopupHandle = FSlateApplication::Get().OnFocusChanging().AddLambda([this](const FFocusEvent& FocusEvent, const FWeakWidgetPath& WeakWidgetPath
@@ -1097,7 +949,6 @@ void FWaveformTransformationMarkerRenderer::ShowConvertToLoopAtMarker(const FGeo
 					if (OldWidget && !OldWidget->IsHovered() && LockedLoopMenuWindow == OldWidget)
 					{
 						FSlateApplication::Get().RequestDestroyWindow(LockedLoopMenuWindow.ToSharedRef());
-						FSlateApplication::Get().OnFocusChanging().Remove(PopupHandle);
 					}
 				}
 			}
diff --git a/Engine/Plugins/Editor/WaveformEditor/Source/WaveformTransformationsWidgets/Private/WaveformTransformationMarkersRenderer.h b/Engine/Plugins/Editor/WaveformEditor/Source/WaveformTransformationsWidgets/Private/WaveformTransformationMarkersRenderer.h
--- a/Engine/Plugins/Editor/WaveformEditor/Source/WaveformTransformationsWidgets/Private/WaveformTransformationMarkersRenderer.h
+++ b/Engine/Plugins/Editor/WaveformEditor/Source/WaveformTransformationsWidgets/Private/WaveformTransformationMarkersRenderer.h
@@ -22,12 +22,7 @@ public:
 	virtual FCursorReply OnCursorQuery(const FGeometry& MyGeometry, const FPointerEvent& CursorEvent) const override;
 	virtual FReply OnKeyUp(const FGeometry& MyGeometry, const FKeyEvent& InKeyEvent) override;
 
-	virtual void SetPropertyHandles(const TArray<TSharedRef<IPropertyHandle>>& InPropertyHandles) override;
-
-	void ModifyMarkerLoopRegion(ELoopModificationControls Modification);
-	void CycleMarkerLoopRegion(ELoopModificationControls Modification);
-	FSoundWaveCuePoint* GetSelectedMarker() const;
-	void SelectLoopRegionByKeyboard(const FKey& PressedKey);
+	virtual void SetWaveformTransformation(TObjectPtr<UWaveformTransformationBase> InTransformation) override;
 
 private:
 	int32 DrawMarkerHandles(const FGeometry& AllottedGeometry, FSlateWindowElementList& OutDrawElements, int32 LayerId) const;
@@ -57,6 +52,7 @@ private:
 		LoopHandle,
 		MarkerHandle,
 		RightClickMarker,
+		DeselectRegion,
 	} MarkerInteractionType = EMarkerInteractionType::None;
 
 	class FMarkerInInteractionRange
@@ -100,18 +96,10 @@ private:
 	FLinearColor LabelTextColor = WaveformTransformationWidgetSharedDefaults::DefaultLabelTextColor;
 	float LabelFontSize = WaveformTransformationWidgetSharedDefaults::DefaultLabelFontSize;
 
-	bool bLoopRegionPreviewNeedsRefresh = false;
-
-	TSharedPtr<IPropertyHandle> MarkersHandle;
-	TSharedPtr<IPropertyHandle> StartLoopTimeHandle;
-	TSharedPtr<IPropertyHandle> EndLoopTimeHandle;
-	TSharedPtr<IPropertyHandle> IsPreviewingLoopRegionHandle;
+	TStrongObjectPtr<UWaveformTransformationMarkers> StrongMarkersTransformation = nullptr;
 
 	FDelegateHandle PopupHandle;
 	FDelegateHandle ApplicationActivationStateHandle;
 
 	TWeakPtr<SWindow> CreateLoopMenuWindow;
-	const TArray<int32> SizeIncrements = { 1, 10, 50, 100 };
-	const TArray<FKey> IndexKeyCodes = { EKeys::One, EKeys::Two, EKeys::Three, EKeys::Four, EKeys::Five, EKeys::Six, EKeys::Seven, EKeys::Eight, EKeys::Nine };
-	int32 SelectedIncrement = 0;
 };
diff --git a/Engine/Plugins/Editor/WaveformEditor/Source/WaveformTransformationsWidgets/Private/WaveformTransformationRendererBase.cpp b/Engine/Plugins/Editor/WaveformEditor/Source/WaveformTransformationsWidgets/Private/WaveformTransformationRendererBase.cpp
--- a/Engine/Plugins/Editor/WaveformEditor/Source/WaveformTransformationsWidgets/Private/WaveformTransformationRendererBase.cpp
+++ b/Engine/Plugins/Editor/WaveformEditor/Source/WaveformTransformationsWidgets/Private/WaveformTransformationRendererBase.cpp
@@ -55,11 +55,14 @@ void FWaveformTransformationRendererBase::SetTransformationWaveInfo(const FWavef
 
 void FWaveformTransformationRendererBase::SetPropertyHandles(const TArray<TSharedRef<IPropertyHandle>>& InPropertyHandles)
 {
+	PRAGMA_DISABLE_DEPRECATION_WARNINGS
 	CachedPropertyHandles = InPropertyHandles;
+	PRAGMA_ENABLE_DEPRECATION_WARNINGS
 }
 
-TSharedPtr<IPropertyHandle> FWaveformTransformationRendererBase::GetPropertyHandle(const FName& PropertyName) const 
+TSharedPtr<IPropertyHandle> FWaveformTransformationRendererBase::GetPropertyHandle(const FName& PropertyName) const
 {
+	PRAGMA_DISABLE_DEPRECATION_WARNINGS
 	TSharedPtr<IPropertyHandle> HandlePtr = nullptr;
 
 	for (const TSharedRef<IPropertyHandle>& PropertyHandle : CachedPropertyHandles)
@@ -70,6 +73,8 @@ TSharedPtr<IPropertyHandle> FWaveformTransformationRendererBase::GetPropertyHand
 			return PropertyHandle;
 		}
 	}
+
+	PRAGMA_ENABLE_DEPRECATION_WARNINGS
 	check(HandlePtr)
 	return HandlePtr;
 }
diff --git a/Engine/Plugins/Editor/WaveformEditor/Source/WaveformTransformationsWidgets/Private/WaveformTransformationTrimFadeRenderer.cpp b/Engine/Plugins/Editor/WaveformEditor/Source/WaveformTransformationsWidgets/Private/WaveformTransformationTrimFadeRenderer.cpp
--- a/Engine/Plugins/Editor/WaveformEditor/Source/WaveformTransformationsWidgets/Private/WaveformTransformationTrimFadeRenderer.cpp
+++ b/Engine/Plugins/Editor/WaveformEditor/Source/WaveformTransformationsWidgets/Private/WaveformTransformationTrimFadeRenderer.cpp
@@ -4,11 +4,13 @@
 
 #include "AudioWidgetsStyle.h"
 #include "Brushes/SlateRoundedBoxBrush.h"
+#include "Editor.h"
 #include "Framework/Application/SlateApplication.h"
 #include "Layout/Geometry.h"
 #include "PropertyHandle.h"
 #include "Rendering/DrawElements.h"
-#include "WaveformTransformationTrimFade.h"
+#include "Sound/SoundWave.h"
+#include "WaveformTransformationLog.h"
 #include "Widgets/Images/SImage.h"
 #include "Widgets/Input/SButton.h"
 #include "Widgets/SWidget.h"
@@ -131,7 +133,8 @@ int32 FWaveformTransformationTrimFadeRenderer::DrawFadeCurves(const FGeometry& A
 
 void FWaveformTransformationTrimFadeRenderer::GenerateFadeCurves(const FGeometry& AllottedGeometry)
 {
-	const float FadeInFrames = GetPropertyValue<float>(StartFadeTimeHandle) * TransformationWaveInfo.SampleRate;
+	check(StrongTrimFade);
+	const float FadeInFrames = StrongTrimFade->StartFadeTime * TransformationWaveInfo.SampleRate;
 	const uint32 FadeInPixelLength = FadeInFrames * PixelsPerFrame;
 	FadeInStartX = FMath::RoundToInt32(StartTimeHandleX);
 	FadeInEndX = FMath::RoundToInt32(FMath::Clamp(StartTimeHandleX + FadeInPixelLength, StartTimeHandleX, EndTimeHandleX));
@@ -142,14 +145,14 @@ void FWaveformTransformationTrimFadeRenderer::GenerateFadeCurves(const FGeometry
 	for (uint32 Pixel = 0; Pixel < DisplayedFadeInPixelLength; ++Pixel)
 	{
 		const double FadeFraction = (float)Pixel / FadeInPixelLength;
-		const double CurveFunction = UWaveformTransformationTrimFade::GetFadeInCurveValue(GetPropertyValue<float>(StartFadeCurveHandle), FadeFraction, GetPropertyValue<float>(StartSCurveSharpnessHandle));
+		const double CurveFunction = UWaveformTransformationTrimFade::GetFadeInCurveValue(StrongTrimFade->StartFadeCurve, FadeFraction, StrongTrimFade->StartSCurveSharpness);
 		const double CurveValue = Pixel != FadeInPixelLength - 1 ? 1.f - CurveFunction : 0.f;
 
 		const uint32 XCoordinate = Pixel + FadeInStartX;
 		FadeInCurvePoints[Pixel] = FVector2D(XCoordinate, CurveValue * AllottedGeometry.Size.Y);
 	}
 
-	const float FadeOutFrames = GetPropertyValue<float>(EndFadeTimeHandle) * TransformationWaveInfo.SampleRate;
+	const float FadeOutFrames = StrongTrimFade->EndFadeTime * TransformationWaveInfo.SampleRate;
 	const float FadeOutPixelLength = FadeOutFrames * PixelsPerFrame;
 	FadeOutStartX = FMath::RoundToInt32(FMath::Clamp(EndTimeHandleX - FadeOutPixelLength, StartTimeHandleX, EndTimeHandleX));
 	FadeOutEndX = FMath::RoundToInt32(EndTimeHandleX);
@@ -161,7 +164,7 @@ void FWaveformTransformationTrimFadeRenderer::GenerateFadeCurves(const FGeometry
 	for (uint32 Pixel = 0; Pixel < DisplayedFadeOutPixelLength; ++Pixel)
 	{
 		const double FadeFraction = (float)(Pixel + FadeOutPixelOffset) / FadeOutPixelLength;
-		const double CurveFunction = UWaveformTransformationTrimFade::GetFadeOutCurveValue(GetPropertyValue<float>(EndFadeCurveHandle), FadeFraction, GetPropertyValue<float>(EndSCurveSharpnessHandle));
+		const double CurveFunction = UWaveformTransformationTrimFade::GetFadeOutCurveValue(StrongTrimFade->EndFadeCurve, FadeFraction, StrongTrimFade->EndSCurveSharpness);
 		const double CurveValue = Pixel != FadeOutPixelLength - 1 ? 1.f - CurveFunction : 1.f;
 
 		const uint32 XCoordinate = Pixel + FadeOutStartX;
@@ -191,31 +194,26 @@ FCursorReply FWaveformTransformationTrimFadeRenderer::OnCursorQuery(const FGeome
 	return FCursorReply::Unhandled();
 }
 
-void FWaveformTransformationTrimFadeRenderer::SetPropertyHandles(const TArray<TSharedRef<IPropertyHandle>>& InPropertyHandles)
+void FWaveformTransformationTrimFadeRenderer::SetWaveformTransformation(TObjectPtr<UWaveformTransformationBase> InTransformation)
 {
-	FWaveformTransformationRendererBase::SetPropertyHandles(InPropertyHandles);
-
-	StartTimeHandle = GetPropertyHandle(GET_MEMBER_NAME_CHECKED(UWaveformTransformationTrimFade, StartTime));
-	EndTimeHandle  = GetPropertyHandle(GET_MEMBER_NAME_CHECKED(UWaveformTransformationTrimFade, EndTime));
-	StartFadeTimeHandle = GetPropertyHandle(GET_MEMBER_NAME_CHECKED(UWaveformTransformationTrimFade, StartFadeTime));
-	StartFadeCurveHandle = GetPropertyHandle(GET_MEMBER_NAME_CHECKED(UWaveformTransformationTrimFade, StartFadeCurve));
-	StartSCurveSharpnessHandle = GetPropertyHandle(GET_MEMBER_NAME_CHECKED(UWaveformTransformationTrimFade, StartSCurveSharpness));
-	EndFadeTimeHandle = GetPropertyHandle(GET_MEMBER_NAME_CHECKED(UWaveformTransformationTrimFade, EndFadeTime));
-	EndFadeCurveHandle = GetPropertyHandle(GET_MEMBER_NAME_CHECKED(UWaveformTransformationTrimFade, EndFadeCurve));
-	EndSCurveSharpnessHandle = GetPropertyHandle(GET_MEMBER_NAME_CHECKED(UWaveformTransformationTrimFade, EndSCurveSharpness));
+	FWaveformTransformationRendererBase::SetWaveformTransformation(InTransformation);
+	StrongTrimFade = TStrongObjectPtr<UWaveformTransformationTrimFade>(CastChecked<UWaveformTransformationTrimFade>(InTransformation.Get()));
 }
 
 FReply FWaveformTransformationTrimFadeRenderer::OnMouseWheel(SWidget& OwnerWidget, const FGeometry& MyGeometry, const FPointerEvent& MouseEvent)
 {
+	check(StrongTrimFade);
 	const FVector2D LocalCursorPosition = GetLocalCursorPosition(MouseEvent, MyGeometry);
 
 	if (IsCursorInFadeInInteractionRange(LocalCursorPosition, MyGeometry))
 	{		
 		const float FadeCurveDelta = MouseEvent.GetWheelDelta() * MouseWheelStep;
-		const float FadeCurveValue = FMath::Clamp(GetPropertyValue<float>(StartFadeCurveHandle) + FadeCurveDelta, 0.f, 10.f);
+		const float FadeCurveValue = FMath::Clamp(StrongTrimFade->StartFadeCurve + FadeCurveDelta, 0.f, 10.f);
 
 		EPropertyValueSetFlags::Type PropertyFlags = EPropertyValueSetFlags::InteractiveChange;
-		StartFadeCurveHandle->SetValue(FadeCurveValue, PropertyFlags);
+		StrongTrimFade->StartFadeCurve = FadeCurveValue;
+
+		StrongTrimFade->PostEditChange();
 
 		return FReply::Handled();
 	}
@@ -223,10 +221,12 @@ FReply FWaveformTransformationTrimFadeRenderer::OnMouseWheel(SWidget& OwnerWidge
 	if (IsCursorInFadeOutInteractionRange(LocalCursorPosition, MyGeometry))
 	{
 		const float FadeCurveDelta = MouseEvent.GetWheelDelta() * MouseWheelStep;
-		const float FadeCurveValue = FMath::Clamp(GetPropertyValue<float>(EndFadeCurveHandle) + FadeCurveDelta, 0.f, 10.f);
+		const float FadeCurveValue = FMath::Clamp(StrongTrimFade->EndFadeCurve + FadeCurveDelta, 0.f, 10.f);
 		
 		EPropertyValueSetFlags::Type PropertyFlags = EPropertyValueSetFlags::InteractiveChange;
-		EndFadeCurveHandle->SetValue(FadeCurveValue, PropertyFlags);
+		StrongTrimFade->EndFadeCurve = FadeCurveValue;
+
+		StrongTrimFade->PostEditChange();
 		
 		return FReply::Handled();
 	}
@@ -252,10 +252,11 @@ double FWaveformTransformationTrimFadeRenderer::ConvertXRatioToTime(const float
 void FWaveformTransformationTrimFadeRenderer::Tick(const FGeometry& AllottedGeometry, const double InCurrentTime, const float InDeltaTime)
 { 
 	check(TransformationWaveInfo.NumChannels > 0);
+	check(StrongTrimFade);
 
 	const float NumFrames = TransformationWaveInfo.TotalNumSamples / TransformationWaveInfo.NumChannels;
-	const double FirstFrame = FMath::Clamp((GetPropertyValue<double>(StartTimeHandle) * TransformationWaveInfo.SampleRate) , 0.f, NumFrames);
-	const double EndFrame = FMath::Clamp((GetPropertyValue<double>(EndTimeHandle) * TransformationWaveInfo.SampleRate), FirstFrame, NumFrames);
+	const double FirstFrame = FMath::Clamp((StrongTrimFade->StartTime * TransformationWaveInfo.SampleRate) , 0.f, NumFrames);
+	const double EndFrame = FMath::Clamp((StrongTrimFade->EndTime * TransformationWaveInfo.SampleRate), FirstFrame, NumFrames);
 
 	check(NumFrames > 0);
 	PixelsPerFrame = FMath::Max(AllottedGeometry.GetLocalSize().X / NumFrames, 0.0);
@@ -281,6 +282,7 @@ void FWaveformTransformationTrimFadeRenderer::UpdateInteractionRange()
 
 FReply FWaveformTransformationTrimFadeRenderer::OnMouseButtonDown(SWidget& OwnerWidget, const FGeometry& MyGeometry, const FPointerEvent& MouseEvent)
 {
+	check(StrongTrimFade);
 	const FVector2D LocalCursorPosition = GetLocalCursorPosition(MouseEvent, MyGeometry);
 
 	const FKey MouseButton = MouseEvent.GetEffectingButton();
@@ -289,7 +291,11 @@ FReply FWaveformTransformationTrimFadeRenderer::OnMouseButtonDown(SWidget& Owner
 
 	if (TrimFadeInteractionType != ETrimFadeInteractionType::None)
 	{
-		BeginTransaction(TEXT("PropertyEditor"), FText::Format(NSLOCTEXT("PropertyEditor", "SetNumericPropertyTransaction", "Edit {0}"), GetPropertyEditedByCurrentInteraction()), nullptr);
+		if (BeginTransaction(TEXT("PropertyEditor"), FText::Format(NSLOCTEXT("PropertyEditor", "SetNumericPropertyTransaction", "Edit {0}"), GetPropertyEditedByCurrentInteraction()), StrongTrimFade.Get()) == INDEX_NONE)
+		{
+			UE_LOG(LogWaveformTransformation, Warning, TEXT("BEGIN TRANSACTION FAILED"));
+		}
+		StrongTrimFade->Modify();
 		return FReply::Handled().CaptureMouse(OwnerWidget.AsShared()).PreventThrottling();
 	}
 
@@ -299,9 +305,11 @@ FReply FWaveformTransformationTrimFadeRenderer::OnMouseButtonDown(SWidget& Owner
 
 FReply FWaveformTransformationTrimFadeRenderer::OnMouseMove(SWidget& OwnerWidget, const FGeometry& MyGeometry, const FPointerEvent& MouseEvent)
 {
+	check(StrongTrimFade);
 	if (MouseEvent.IsMouseButtonDown(EKeys::LeftMouseButton) && TrimFadeInteractionType != ETrimFadeInteractionType::None)
 	{
-		SetPropertyValueDependingOnInteractionType(MouseEvent, MyGeometry, EPropertyValueSetFlags::InteractiveChange | EPropertyValueSetFlags::NotTransactable);
+		SetPropertyValueDependingOnInteractionType(MouseEvent, MyGeometry);
+		StrongTrimFade->OnTransformationRenderChanged.ExecuteIfBound();
 
 		return FReply::Handled().CaptureMouse(OwnerWidget.AsShared());
 	}
@@ -311,113 +319,93 @@ FReply FWaveformTransformationTrimFadeRenderer::OnMouseMove(SWidget& OwnerWidget
 
 FReply FWaveformTransformationTrimFadeRenderer::OnMouseButtonUp(SWidget& OwnerWidget, const FGeometry& MyGeometry, const FPointerEvent& MouseEvent)
 {
+	check(StrongTrimFade);
 	if (TrimFadeInteractionType != ETrimFadeInteractionType::None)
 	{
+		StrongTrimFade->Modify();
+
 		SetPropertyValueDependingOnInteractionType(MouseEvent, MyGeometry);
+
 		EndTransaction();
+		StrongTrimFade->OnTransformationChanged.ExecuteIfBound();
+
 		TrimFadeInteractionType = ETrimFadeInteractionType::None;
+		
 		return FReply::Handled().ReleaseMouseCapture();
 	}
 
 	return FReply::Unhandled();
 }
 
-void FWaveformTransformationTrimFadeRenderer::SetPropertyValueDependingOnInteractionType(const FPointerEvent& MouseEvent, const FGeometry& WidgetGeometry, const EPropertyValueSetFlags::Type Flags /*= EPropertyValueSetFlags::DefaultFlags*/)
+void FWaveformTransformationTrimFadeRenderer::SetPropertyValueDependingOnInteractionType(const FPointerEvent& MouseEvent, const FGeometry& WidgetGeometry)
 {
+	check(StrongTrimFade);
 	const FVector2D LocalCursorPosition = GetLocalCursorPosition(MouseEvent, WidgetGeometry);
 	const float LocalCursorXRatio = FMath::Clamp(LocalCursorPosition.X / WidgetGeometry.GetLocalSize().X, 0.f, 1.f);
 	const double SelectedTime = ConvertXRatioToTime(LocalCursorXRatio);
 
-	TSharedPtr<IPropertyHandle> PropertyHandle;
-
 	switch (TrimFadeInteractionType)
 	{
 	case FWaveformTransformationTrimFadeRenderer::ETrimFadeInteractionType::None:
 		break;
 	case FWaveformTransformationTrimFadeRenderer::ETrimFadeInteractionType::ScrubbingLeftHandle:
-		StartTimeHandle->SetValue(SelectedTime, Flags);
+		StrongTrimFade->StartTime = FMath::Min(SelectedTime, StrongTrimFade->EndTime);
 		break;
 	case FWaveformTransformationTrimFadeRenderer::ETrimFadeInteractionType::ScrubbingRightHandle:
-		EndTimeHandle->SetValue(SelectedTime, Flags);
+		StrongTrimFade->EndTime = FMath::Max(SelectedTime, StrongTrimFade->StartTime);
 		break;
 	case FWaveformTransformationTrimFadeRenderer::ETrimFadeInteractionType::ScrubbingFadeIn:
 		{
-			float StartFadeTimeValue = FMath::Clamp(SelectedTime - GetPropertyValue<double>(StartTimeHandle), 0.f, TNumericLimits<float>().Max());
-			StartFadeTimeHandle->SetValue(StartFadeTimeValue, Flags);
+			float StartFadeTimeValue = FMath::Clamp(SelectedTime - StrongTrimFade->StartTime, 0.f, TNumericLimits<float>().Max());
+			StrongTrimFade->StartFadeTime = StartFadeTimeValue;
 		}
 		break;
 	case FWaveformTransformationTrimFadeRenderer::ETrimFadeInteractionType::ScrubbingFadeOut:
 		{	
-			float EndFadeTimeValue = FMath::Clamp(GetPropertyValue<double>(EndTimeHandle) - SelectedTime, 0.f, TNumericLimits<float>().Max());
-			EndFadeTimeHandle->SetValue(EndFadeTimeValue, Flags);
+			float EndFadeTimeValue = FMath::Clamp(StrongTrimFade->EndTime - SelectedTime, 0.f, TNumericLimits<float>().Max());
+			StrongTrimFade->EndFadeTime = EndFadeTimeValue;
 		}
 		break;
 	case FWaveformTransformationTrimFadeRenderer::ETrimFadeInteractionType::RightClickFadeIn:
-		ShowSelectFadeModeMenuAtCursor(WidgetGeometry, MouseEvent, StartFadeCurveHandle);
+		CurvePropertyToChange = MakeShared<float>(StrongTrimFade->StartFadeCurve);
+		ShowSelectFadeModeMenuAtCursor(WidgetGeometry, MouseEvent);
 		break;
 	case FWaveformTransformationTrimFadeRenderer::ETrimFadeInteractionType::RightClickFadeOut:
-		ShowSelectFadeModeMenuAtCursor(WidgetGeometry, MouseEvent, EndFadeCurveHandle);
+		CurvePropertyToChange = MakeShared<float>(StrongTrimFade->EndFadeCurve);
+		ShowSelectFadeModeMenuAtCursor(WidgetGeometry, MouseEvent);
 		break;
 	default:
 		break;
 	}
 }
 
-void FWaveformTransformationTrimFadeRenderer::ShowSelectFadeModeMenuAtCursor(const FGeometry& MyGeometry, const FPointerEvent& MouseEvent, TSharedPtr<IPropertyHandle> PropertyHandle)
+void FWaveformTransformationTrimFadeRenderer::ShowSelectFadeModeMenuAtCursor(const FGeometry& MyGeometry, const FPointerEvent& MouseEvent)
 {
-	if (!PropertyHandle.IsValid())
+	const FVector2D LocalWindowMaxPosition = MyGeometry.GetAbsolutePosition() + MyGeometry.GetAbsoluteSize();
+	const FVector2D LocalCursorPosition = MouseEvent.GetScreenSpacePosition();
+
+	if (TSharedPtr<SWindow> LockedFadeModeMenuWindow = FadeModeMenuWindow.Pin())
 	{
-		return;
+		FSlateApplication::Get().RequestDestroyWindow(LockedFadeModeMenuWindow.ToSharedRef());
+		FSlateApplication::Get().OnFocusChanging().Remove(PopupHandle); 
+		FSlateApplication::Get().OnApplicationActivationStateChanged().Remove(ApplicationActivationStateHandle);
 	}
 
-	TWeakPtr<IPropertyHandle> WeakStartFadeCurveHandle = StartFadeCurveHandle;
-	TWeakPtr<IPropertyHandle> WeakEndFadeCurveHandle = EndFadeCurveHandle;
-
-	TFunction<FSlateIcon(const EWaveEditorFadeMode&, const TSharedPtr<IPropertyHandle>&)> GetFadeModeIcon = 
-		[WeakStartFadeCurveHandle, WeakEndFadeCurveHandle](const EWaveEditorFadeMode& FadeMode, const TSharedPtr<IPropertyHandle>& PropertyHandle)
-		{
-			TSharedPtr<IPropertyHandle> LockedStartFadeCurveHandle = WeakStartFadeCurveHandle.Pin();
-			TSharedPtr<IPropertyHandle> LockedEndFadeCurveHandle = WeakEndFadeCurveHandle.Pin();
-			if (LockedStartFadeCurveHandle && PropertyHandle == LockedStartFadeCurveHandle)
-			{
-				switch (FadeMode)
-				{
-				case EWaveEditorFadeMode::Linear:
-					return FSlateIcon(FAudioWidgetsStyle::StyleName, "AudioWidgetsStyle.FadeInLinear");
-				case EWaveEditorFadeMode::Exponetial:
-					return FSlateIcon(FAudioWidgetsStyle::StyleName, "AudioWidgetsStyle.FadeInExponential");
-				case EWaveEditorFadeMode::Logarithmic:
-					return FSlateIcon(FAudioWidgetsStyle::StyleName, "AudioWidgetsStyle.FadeInLogarithmic");
-				case EWaveEditorFadeMode::Sigmoid:
-					return FSlateIcon(FAudioWidgetsStyle::StyleName, "AudioWidgetsStyle.FadeInSigmoid");
-				default:
-					return FSlateIcon();
-				}
-			}
-			else if (LockedEndFadeCurveHandle && PropertyHandle == LockedEndFadeCurveHandle)
-			{
-				switch (FadeMode)
-				{
-				case EWaveEditorFadeMode::Linear:
-					return FSlateIcon(FAudioWidgetsStyle::StyleName, "AudioWidgetsStyle.FadeOutLinear");
-				case EWaveEditorFadeMode::Exponetial:
-					return FSlateIcon(FAudioWidgetsStyle::StyleName, "AudioWidgetsStyle.FadeOutExponential");
-				case EWaveEditorFadeMode::Logarithmic:
-					return FSlateIcon(FAudioWidgetsStyle::StyleName, "AudioWidgetsStyle.FadeOutLogarithmic");
-				case EWaveEditorFadeMode::Sigmoid:
-					return FSlateIcon(FAudioWidgetsStyle::StyleName, "AudioWidgetsStyle.FadeOutSigmoid");
-				default:
-					return FSlateIcon();
-				}
-			}
-
-			return FSlateIcon();
-			
-		};
+	TSharedRef<SVerticalBox> MenuContent = SNew(SVerticalBox);
 
-	const FVector2D LocalCursorPosition = MouseEvent.GetScreenSpacePosition();
+	TSharedRef<SWindow> MenuWindow = SNew(SWindow)
+		.AutoCenter(EAutoCenter::None)
+		.ScreenPosition(LocalCursorPosition)
+		.SizingRule(ESizingRule::Autosized)
+		.SupportsTransparency(EWindowTransparency::None)
+		.IsPopupWindow(true)
+		.CreateTitleBar(false)
+		[
+			MenuContent
+		];
 
-	TSharedRef<SVerticalBox> MenuContent = SNew(SVerticalBox);
+	FadeModeMenuWindow = MenuWindow;
+	TWeakObjectPtr<UWaveformTransformationTrimFade> InWeakTrimFade = StrongTrimFade.Get();
 
 	for (const TPair<EWaveEditorFadeMode, float>& FadeOptionPair : UWaveformTransformationTrimFade::FadeModeToCurveValueMap)
 	{
@@ -426,11 +414,37 @@ void FWaveformTransformationTrimFadeRenderer::ShowSelectFadeModeMenuAtCursor(con
 			.AutoHeight()
 			[
 				SNew(SButton)
-					.OnClicked_Lambda([this, PropertyHandle, FadeOptionPair]() -> FReply 
+					.OnClicked_Lambda([InWeakFadeModeMenuWindow = FadeModeMenuWindow, InWeakTrimFade, InteractionType = TrimFadeInteractionType, FadeOptionPair, PropertyName = GetPropertyEditedByCurrentInteraction()]() -> FReply
 					{
-						PropertyHandle->SetValue(FadeOptionPair.Value);
+						if (TStrongObjectPtr<UWaveformTransformationTrimFade> PinnedTrimFade = InWeakTrimFade.Pin())
+						{
+							if (GEditor && GEditor->Trans)
+							{
+								GEditor->BeginTransaction(TEXT("PropertyEditor"), FText::Format(NSLOCTEXT("PropertyEditor", "SetCurveProperty", "Set Curve"), PropertyName), PinnedTrimFade.Get());
+							}
+
+							PinnedTrimFade->Modify();
+
+							if (InteractionType == FWaveformTransformationTrimFadeRenderer::ETrimFadeInteractionType::RightClickFadeIn)
+							{
+								PinnedTrimFade->StartFadeCurve = FadeOptionPair.Value;
+							}
+							else if (InteractionType == FWaveformTransformationTrimFadeRenderer::ETrimFadeInteractionType::RightClickFadeOut)
+							{
+								PinnedTrimFade->EndFadeCurve = FadeOptionPair.Value;
+							}
+
+							if (GEditor && GEditor->Trans)
+							{
+								GEditor->EndTransaction();
+							}
+
+							PinnedTrimFade->PostEditChange();
+							PinnedTrimFade->OnTransformationChanged.ExecuteIfBound();
+						}
 
-						if (TSharedPtr<SWindow> LockedFadeModeMenuWindow = FadeModeMenuWindow.Pin())
+						TSharedPtr<SWindow> LockedFadeModeMenuWindow = InWeakFadeModeMenuWindow.Pin();
+						if (LockedFadeModeMenuWindow)
 						{
 							FSlateApplication::Get().RequestDestroyWindow(LockedFadeModeMenuWindow.ToSharedRef());
 						}
@@ -445,7 +459,7 @@ void FWaveformTransformationTrimFadeRenderer::ShowSelectFadeModeMenuAtCursor(con
 						.HAlign(HAlign_Left)
 						[
 							SNew(SImage)
-							.Image(GetFadeModeIcon(FadeOptionPair.Key, PropertyHandle).GetIcon())
+							.Image(GetFadeModeIcon(FadeOptionPair.Key).GetIcon())
 						]
 						+ SHorizontalBox::Slot()
 						.VAlign(VAlign_Center)
@@ -459,54 +473,94 @@ void FWaveformTransformationTrimFadeRenderer::ShowSelectFadeModeMenuAtCursor(con
 			];
 	}
 
-	TSharedRef<SWindow> MenuWindow = SNew(SWindow)
-		.AutoCenter(EAutoCenter::None)
-		.ScreenPosition(LocalCursorPosition)
-		.SizingRule(ESizingRule::Autosized)
-		.SupportsTransparency(EWindowTransparency::None)
-		.IsPopupWindow(true)
-		.CreateTitleBar(false)
-		[
-			MenuContent
-		];
+	FSlateApplication::Get().AddWindow(MenuWindow);
 
-	FadeModeMenuWindow = MenuWindow;
+	FTSTicker::GetCoreTicker().AddTicker(FTickerDelegate::CreateLambda([InFadeModeMenuWindow = FadeModeMenuWindow, LocalCursorPosition, LocalWindowMaxPosition](float DeltaTime)
+		{
+			if (TSharedPtr<SWindow> LockedFadeModeMenuWindow = InFadeModeMenuWindow.Pin())
+			{
+				const FVector2D WindowSize = LockedFadeModeMenuWindow->GetSizeInScreen();
 
-	FSlateApplication::Get().AddWindow(MenuWindow);
+				if (WindowSize.X > 0.0 && LocalCursorPosition.X + WindowSize.X > LocalWindowMaxPosition.X)
+				{
+					check(WindowSize.Y > 0.0);
 
-	FSlateApplication::Get().OnFocusChanging().Remove(PopupHandle);
-	FSlateApplication::Get().OnApplicationActivationStateChanged().Remove(ApplicationActivationStateHandle);
+					const double XPosition = FMath::Max(0.0, LocalCursorPosition.X - WindowSize.X);
+					LockedFadeModeMenuWindow->MoveWindowTo(FVector2D(XPosition, LocalCursorPosition.Y));
+				}
+			}
 
-	ApplicationActivationStateHandle = FSlateApplication::Get().OnApplicationActivationStateChanged().AddLambda([this](bool isActive)
+			return false; // return false to only execute once
+		}));
+
+	ApplicationActivationStateHandle = FSlateApplication::Get().OnApplicationActivationStateChanged().AddLambda([InFadeModeMenuWindow = FadeModeMenuWindow, InApplicationActivationStateHandle = ApplicationActivationStateHandle](bool isActive)
 		{
 			if (!isActive)
 			{
-				if (TSharedPtr<SWindow> LockedFadeModeMenuWindow = FadeModeMenuWindow.Pin())
+				if (TSharedPtr<SWindow> LockedFadeModeMenuWindow = InFadeModeMenuWindow.Pin())
 				{
-					FSlateApplication::Get().RequestDestroyWindow(LockedFadeModeMenuWindow.ToSharedRef());
-					FSlateApplication::Get().OnApplicationActivationStateChanged().Remove(ApplicationActivationStateHandle);
+					FSlateApplication::Get().RequestDestroyWindow(LockedFadeModeMenuWindow.ToSharedRef());	
+					FSlateApplication::Get().OnApplicationActivationStateChanged().Remove(InApplicationActivationStateHandle);
 				}
 			}
 		});
 
 	// If focus is lost on the popup, destroy it to prevent popups hanging around
-	PopupHandle = FSlateApplication::Get().OnFocusChanging().AddLambda([this](const FFocusEvent& FocusEvent, const FWeakWidgetPath& WeakWidgetPath
+	PopupHandle = FSlateApplication::Get().OnFocusChanging().AddLambda([InFadeModeMenuWindow = FadeModeMenuWindow, InPopupHandle = PopupHandle](const FFocusEvent& FocusEvent, const FWeakWidgetPath& WeakWidgetPath
 		, const TSharedPtr<SWidget>& OldWidget, const FWidgetPath& WidgetPath, const TSharedPtr<SWidget>& NewWidget)
 		{
-			if (FadeModeMenuWindow != nullptr && FadeModeMenuWindow.IsValid())
+			if (InFadeModeMenuWindow != nullptr && InFadeModeMenuWindow.IsValid())
 			{
-				if (TSharedPtr<SWindow> LockedFadeModeMenuWindow = FadeModeMenuWindow.Pin())
+				if (TSharedPtr<SWindow> LockedFadeModeMenuWindow = InFadeModeMenuWindow.Pin())
 				{
 					if (OldWidget && !OldWidget->IsHovered() && LockedFadeModeMenuWindow == OldWidget)
 					{
 						FSlateApplication::Get().RequestDestroyWindow(LockedFadeModeMenuWindow.ToSharedRef());
-						FSlateApplication::Get().OnFocusChanging().Remove(PopupHandle);
+						FSlateApplication::Get().OnFocusChanging().Remove(InPopupHandle);
 					}
 				}
 			}
 		});
 }
 
+const FSlateIcon FWaveformTransformationTrimFadeRenderer::GetFadeModeIcon(const EWaveEditorFadeMode& FadeMode)
+{
+	if (TrimFadeInteractionType == FWaveformTransformationTrimFadeRenderer::ETrimFadeInteractionType::RightClickFadeIn)
+	{
+		switch (FadeMode)
+		{
+		case EWaveEditorFadeMode::Linear:
+			return FSlateIcon(FAudioWidgetsStyle::StyleName, "AudioWidgetsStyle.FadeInLinear");
+		case EWaveEditorFadeMode::Exponetial:
+			return FSlateIcon(FAudioWidgetsStyle::StyleName, "AudioWidgetsStyle.FadeInExponential");
+		case EWaveEditorFadeMode::Logarithmic:
+			return FSlateIcon(FAudioWidgetsStyle::StyleName, "AudioWidgetsStyle.FadeInLogarithmic");
+		case EWaveEditorFadeMode::Sigmoid:
+			return FSlateIcon(FAudioWidgetsStyle::StyleName, "AudioWidgetsStyle.FadeInSigmoid");
+		default:
+			return FSlateIcon();
+		}
+	}
+	else if (TrimFadeInteractionType == FWaveformTransformationTrimFadeRenderer::ETrimFadeInteractionType::RightClickFadeOut)
+	{
+		switch (FadeMode)
+		{
+		case EWaveEditorFadeMode::Linear:
+			return FSlateIcon(FAudioWidgetsStyle::StyleName, "AudioWidgetsStyle.FadeOutLinear");
+		case EWaveEditorFadeMode::Exponetial:
+			return FSlateIcon(FAudioWidgetsStyle::StyleName, "AudioWidgetsStyle.FadeOutExponential");
+		case EWaveEditorFadeMode::Logarithmic:
+			return FSlateIcon(FAudioWidgetsStyle::StyleName, "AudioWidgetsStyle.FadeOutLogarithmic");
+		case EWaveEditorFadeMode::Sigmoid:
+			return FSlateIcon(FAudioWidgetsStyle::StyleName, "AudioWidgetsStyle.FadeOutSigmoid");
+		default:
+			return FSlateIcon();
+		}
+	}
+
+	return FSlateIcon();
+}
+
 FWaveformTransformationTrimFadeRenderer::ETrimFadeInteractionType FWaveformTransformationTrimFadeRenderer::GetInteractionTypeFromCursorPosition(const FVector2D& InLocalCursorPosition, const FKey MouseButton, const FGeometry& WidgetGeometry) const
 {
 	if (MouseButton == EKeys::LeftMouseButton)
diff --git a/Engine/Plugins/Editor/WaveformEditor/Source/WaveformTransformationsWidgets/Private/WaveformTransformationTrimFadeRenderer.h b/Engine/Plugins/Editor/WaveformEditor/Source/WaveformTransformationsWidgets/Private/WaveformTransformationTrimFadeRenderer.h
--- a/Engine/Plugins/Editor/WaveformEditor/Source/WaveformTransformationsWidgets/Private/WaveformTransformationTrimFadeRenderer.h
+++ b/Engine/Plugins/Editor/WaveformEditor/Source/WaveformTransformationsWidgets/Private/WaveformTransformationTrimFadeRenderer.h
@@ -3,6 +3,7 @@
 #pragma once
 
 #include "WaveformTransformationRendererBase.h"
+#include "WaveformTransformationTrimFade.h"
 
 class FWaveformTransformationTrimFadeRenderer : public FWaveformTransformationRendererBase
 {
@@ -18,7 +19,7 @@ public:
 	virtual FReply OnMouseWheel(SWidget& OwnerWidget, const FGeometry& MyGeometry, const FPointerEvent& MouseEvent) override;
 	virtual FCursorReply OnCursorQuery(const FGeometry& MyGeometry, const FPointerEvent& CursorEvent) const override;
 
-	virtual void SetPropertyHandles(const TArray<TSharedRef<IPropertyHandle>>& InPropertyHandles) override;
+	virtual void SetWaveformTransformation(TObjectPtr<UWaveformTransformationBase> InTransformation) override;
 
 private:
 	int32 DrawTrimHandles(const FGeometry& AllottedGeometry, FSlateWindowElementList& OutDrawElements, int32 LayerId) const;
@@ -61,19 +62,14 @@ private:
 	} TrimFadeInteractionType = ETrimFadeInteractionType::None;
 
 	ETrimFadeInteractionType GetInteractionTypeFromCursorPosition(const FVector2D& InLocalCursorPosition, const FKey MouseButton, const FGeometry& WidgetGeometry) const;
-	void SetPropertyValueDependingOnInteractionType(const FPointerEvent& MouseEvent, const FGeometry& WidgetGeometry, const EPropertyValueSetFlags::Type Flags = EPropertyValueSetFlags::DefaultFlags);
-	void ShowSelectFadeModeMenuAtCursor(const FGeometry& MyGeometry, const FPointerEvent& MouseEvent, TSharedPtr<IPropertyHandle> PropertyHandle);
-
-	TWeakPtr<SWindow> FadeModeMenuWindow;
-
-	TSharedPtr<IPropertyHandle> StartTimeHandle;
-	TSharedPtr<IPropertyHandle> EndTimeHandle;
-	TSharedPtr<IPropertyHandle> StartFadeTimeHandle;
-	TSharedPtr<IPropertyHandle> StartFadeCurveHandle;
-	TSharedPtr<IPropertyHandle> StartSCurveSharpnessHandle;
-	TSharedPtr<IPropertyHandle> EndFadeTimeHandle;
-	TSharedPtr<IPropertyHandle> EndFadeCurveHandle;
-	TSharedPtr<IPropertyHandle> EndSCurveSharpnessHandle;
+	void SetPropertyValueDependingOnInteractionType(const FPointerEvent& MouseEvent, const FGeometry& WidgetGeometry);
+	void ShowSelectFadeModeMenuAtCursor(const FGeometry& MyGeometry, const FPointerEvent& MouseEvent);
+	const FSlateIcon GetFadeModeIcon(const EWaveEditorFadeMode& FadeMode);
+
+	TWeakPtr<SWindow> FadeModeMenuWindow = nullptr;
+	TSharedPtr<float> CurvePropertyToChange = nullptr;
+
+	TStrongObjectPtr<UWaveformTransformationTrimFade> StrongTrimFade = nullptr;
 
 	FDelegateHandle PopupHandle;
 	FDelegateHandle ApplicationActivationStateHandle;
diff --git a/Engine/Plugins/Editor/WaveformEditor/Source/WaveformTransformationsWidgets/Private/WaveformTransformationsWidgetsModule.cpp b/Engine/Plugins/Editor/WaveformEditor/Source/WaveformTransformationsWidgets/Private/WaveformTransformationsWidgetsModule.cpp
--- a/Engine/Plugins/Editor/WaveformEditor/Source/WaveformTransformationsWidgets/Private/WaveformTransformationsWidgetsModule.cpp
+++ b/Engine/Plugins/Editor/WaveformEditor/Source/WaveformTransformationsWidgets/Private/WaveformTransformationsWidgetsModule.cpp
@@ -3,12 +3,15 @@
 #include "WaveformTransformationsWidgetsModule.h"
 
 #include "Modules/ModuleManager.h"
+#include "WaveformTransformationLog.h"
 #include "WaveformTransformationRendererMapper.h"
 #include "WaveformTransformationTrimFade.h"
 #include "WaveformTransformationTrimFadeRenderer.h"
 #include "WaveformTransformationMarkers.h"
 #include "WaveformTransformationMarkersRenderer.h"
 
+DEFINE_LOG_CATEGORY(LogWaveformTransformation);
+
 void FWaveformTransformationsWidgetsModule::StartupModule()
 {
 	FWaveformTransformationRendererMapper& RendererMapper = FWaveformTransformationRendererMapper::Get();
diff --git a/Engine/Plugins/Editor/WaveformEditor/Source/WaveformTransformationsWidgets/Public/IWaveformTransformationRenderer.h b/Engine/Plugins/Editor/WaveformEditor/Source/WaveformTransformationsWidgets/Public/IWaveformTransformationRenderer.h
--- a/Engine/Plugins/Editor/WaveformEditor/Source/WaveformTransformationsWidgets/Public/IWaveformTransformationRenderer.h
+++ b/Engine/Plugins/Editor/WaveformEditor/Source/WaveformTransformationsWidgets/Public/IWaveformTransformationRenderer.h
@@ -48,6 +48,9 @@ public:
 	/* Gets info about the transformation being rendered */
 	virtual void SetTransformationWaveInfo(const FWaveformTransformationRenderInfo& InWaveInfo) = 0;
 	
+	virtual void SetWaveformTransformation(TObjectPtr<UWaveformTransformationBase> InTransformation) = 0;
+
+	UE_DEPRECATED(5.7, "PropertyHandles have been deprecated.")
 	/* Gets the transformations properties. IPropertyHandle is used to get/set properties to have full support of change notifiers, editor undo ecc  */
 	virtual void SetPropertyHandles(const TArray<TSharedRef<IPropertyHandle>>& InPropertyHandles) = 0;
 };
diff --git a/Engine/Plugins/Editor/WaveformEditor/Source/WaveformTransformationsWidgets/Public/WaveformTransformationRendererBase.h b/Engine/Plugins/Editor/WaveformEditor/Source/WaveformTransformationsWidgets/Public/WaveformTransformationRendererBase.h
--- a/Engine/Plugins/Editor/WaveformEditor/Source/WaveformTransformationsWidgets/Public/WaveformTransformationRendererBase.h
+++ b/Engine/Plugins/Editor/WaveformEditor/Source/WaveformTransformationsWidgets/Public/WaveformTransformationRendererBase.h
@@ -26,12 +26,17 @@ public:
 	UE_API virtual FReply OnKeyUp(const FGeometry& MyGeometry, const FKeyEvent& InKeyEvent) override;
 
 	UE_API virtual void SetTransformationWaveInfo(const FWaveformTransformationRenderInfo& InWaveInfo) override;
+	UE_API virtual void SetWaveformTransformation(TObjectPtr<UWaveformTransformationBase> InTransformation) override {}
+
+	UE_DEPRECATED(5.7, "PropertyHandles have been deprecated.")
 	UE_API virtual void SetPropertyHandles(const TArray<TSharedRef<IPropertyHandle>>& InPropertyHandles) override;
 
 protected:
+	UE_DEPRECATED(5.7, "PropertyHandles have been deprecated.")
 	UE_API TSharedPtr<IPropertyHandle> GetPropertyHandle(const FName& PropertyName) const;
 
 	template<typename T>
+	UE_DEPRECATED(5.7, "PropertyHandles have been deprecated.")
 	T GetPropertyValue(const TSharedPtr<IPropertyHandle>& PropertyHandle)
 	{
 		check(PropertyHandle.IsValid());
@@ -41,13 +46,13 @@ protected:
 	}
 
 	template<typename T>
+	UE_DEPRECATED(5.7, "PropertyHandles have been deprecated.")
 	void SetPropertyValue(const FName& PropertyName, const T Value, const EPropertyValueSetFlags::Type Flags = EPropertyValueSetFlags::DefaultFlags)
 	{
 		TSharedPtr<IPropertyHandle> Handle = GetPropertyHandle(PropertyName);
 		Handle->SetValue(Value, Flags);
 	}
 
-
 	UE_API int32 BeginTransaction(const TCHAR* TransactionContext, const FText& Description, UObject* PrimaryObject);
 	UE_API int32 EndTransaction();
 
@@ -58,6 +63,7 @@ protected:
 	static constexpr float MouseWheelStep = 0.25f;
 
 private:
+	UE_DEPRECATED(5.7, "PropertyHandles have been deprecated.")
 	TArray<TSharedRef<IPropertyHandle>> CachedPropertyHandles;
 };
 
diff --git a/Engine/Source/Runtime/AudioExtensions/Private/IWaveformTransformation.cpp b/Engine/Source/Runtime/AudioExtensions/Private/IWaveformTransformation.cpp
--- a/Engine/Source/Runtime/AudioExtensions/Private/IWaveformTransformation.cpp
+++ b/Engine/Source/Runtime/AudioExtensions/Private/IWaveformTransformation.cpp
@@ -21,3 +21,20 @@ TArray<Audio::FTransformationPtr> UWaveformTransformationChain::CreateTransforma
 	return TransformationPtrs;
 }
 
+#if WITH_EDITOR
+void UWaveformTransformationBase::PostEditUndo()
+{
+	Super::PostEditUndo();
+
+	OnTransformationChanged.ExecuteIfBound();
+}
+
+void UWaveformTransformationBase::NotifyPropertyChange(FProperty* Property)
+{
+	if (Property)
+	{
+		FPropertyChangedEvent PropertyChangedEvent(Property);
+		PostEditChangeProperty(PropertyChangedEvent);
+	}
+}
+#endif //WITH_EDITOR
diff --git a/Engine/Source/Runtime/AudioExtensions/Public/IWaveformTransformation.h b/Engine/Source/Runtime/AudioExtensions/Public/IWaveformTransformation.h
--- a/Engine/Source/Runtime/AudioExtensions/Public/IWaveformTransformation.h
+++ b/Engine/Source/Runtime/AudioExtensions/Public/IWaveformTransformation.h
@@ -154,9 +154,22 @@ public:
 #if WITH_EDITOR
 	virtual void OverwriteSoundWaveData(USoundWave& InOutSoundWave) {};
 	virtual void GetTransformationInfo(FWaveformTransformationInfo& InOutTransformationInfo) const {};
+
+	AUDIOEXTENSIONS_API virtual void PostEditUndo() override;
+
+	void NotifyPropertyChange(FProperty* Property);
 #endif //WITH_EDITOR
 
 	virtual bool IsEditorOnly() const override { return true; }
+
+
+	// Execute when the soundwave needs a recook to apply new transformation changes
+	DECLARE_DELEGATE(FOnTransformationChanged);
+	FOnTransformationChanged OnTransformationChanged;
+
+	// Execute when only the rendering needs to be updated, but can wait for a recook
+	DECLARE_DELEGATE(FOnTransformationRenderChanged);
+	FOnTransformationRenderChanged OnTransformationRenderChanged;
 };
 
 // Object that holds an ordered list of transformations to perform on a sound wave
