LogWaveformTransformation build error fix

diff --git a/Engine/Plugins/Editor/WaveformEditor/Source/WaveformTransformationsWidgets/Private/WaveformTransformationTrimFadeRenderer.cpp b/Engine/Plugins/Editor/WaveformEditor/Source/WaveformTransformationsWidgets/Private/WaveformTransformationTrimFadeRenderer.cpp
--- a/Engine/Plugins/Editor/WaveformEditor/Source/WaveformTransformationsWidgets/Private/WaveformTransformationTrimFadeRenderer.cpp
+++ b/Engine/Plugins/Editor/WaveformEditor/Source/WaveformTransformationsWidgets/Private/WaveformTransformationTrimFadeRenderer.cpp
@@ -10,7 +10,6 @@
 #include "PropertyHandle.h"
 #include "Rendering/DrawElements.h"
 #include "Sound/SoundWave.h"
-#include "WaveformTransformationLog.h"
 #include "Widgets/Images/SImage.h"
 #include "Widgets/Input/SButton.h"
 #include "Widgets/SWidget.h"
@@ -293,7 +292,7 @@ FReply FWaveformTransformationTrimFadeRenderer::OnMouseButtonDown(SWidget& Owner
 	{
 		if (BeginTransaction(TEXT("PropertyEditor"), FText::Format(NSLOCTEXT("PropertyEditor", "SetNumericPropertyTransaction", "Edit {0}"), GetPropertyEditedByCurrentInteraction()), StrongTrimFade.Get()) == INDEX_NONE)
 		{
-			UE_LOG(LogWaveformTransformation, Warning, TEXT("BEGIN TRANSACTION FAILED"));
+			//TODO: Log a warning
 		}
 		StrongTrimFade->Modify();
 		return FReply::Handled().CaptureMouse(OwnerWidget.AsShared()).PreventThrottling();
diff --git a/Engine/Plugins/Editor/WaveformEditor/Source/WaveformTransformationsWidgets/Private/WaveformTransformationsWidgetsModule.cpp b/Engine/Plugins/Editor/WaveformEditor/Source/WaveformTransformationsWidgets/Private/WaveformTransformationsWidgetsModule.cpp
--- a/Engine/Plugins/Editor/WaveformEditor/Source/WaveformTransformationsWidgets/Private/WaveformTransformationsWidgetsModule.cpp
+++ b/Engine/Plugins/Editor/WaveformEditor/Source/WaveformTransformationsWidgets/Private/WaveformTransformationsWidgetsModule.cpp
@@ -10,8 +10,6 @@
 #include "WaveformTransformationMarkers.h"
 #include "WaveformTransformationMarkersRenderer.h"
 
-DEFINE_LOG_CATEGORY(LogWaveformTransformation);
-
 void FWaveformTransformationsWidgetsModule::StartupModule()
 {
 	FWaveformTransformationRendererMapper& RendererMapper = FWaveformTransformationRendererMapper::Get();
