Waveform Editor: Playback hotkey adjustments

diff --git a/Engine/Plugins/Editor/WaveformEditor/Source/WaveformEditor/Private/WaveformEditor.cpp b/Engine/Plugins/Editor/WaveformEditor/Source/WaveformEditor/Private/WaveformEditor.cpp
--- a/Engine/Plugins/Editor/WaveformEditor/Source/WaveformEditor/Private/WaveformEditor.cpp
+++ b/Engine/Plugins/Editor/WaveformEditor/Source/WaveformEditor/Private/WaveformEditor.cpp
@@ -7,6 +7,7 @@
 #include "AssetDefinitionRegistry.h"
 #include "Components/AudioComponent.h"
 #include "EditorReimportHandler.h"
+#include "Framework/Application/SlateApplication.h"
 #include "Framework/Notifications/NotificationManager.h"
 #include "Misc/TransactionObjectEvent.h"
 #include "PropertyEditorModule.h"
@@ -180,6 +181,8 @@ bool FWaveformEditor::Init(const EToolkitMode::Type Mode, const TSharedPtr<ITool
 		}
 	}
 
+	FSlateApplication::Get().SetKeyboardFocus(TransformationsDetails);
+
 	return bIsInitialized;
 }
 
@@ -863,6 +866,12 @@ void FWaveformEditor::ExecuteOverwriteTransformations()
 	}
 }
 
+void FWaveformEditor::ResetPlaybackToStart()
+{
+	HandlePlayheadScrub(0, false);
+	TransportCoordinator->SetProgressRatio(0);
+}
+
 void FWaveformEditor::RegisterTabSpawners(const TSharedRef<FTabManager>& InTabManager)
 {
 	if (TransformationsDetails == nullptr)
@@ -1274,6 +1283,10 @@ bool FWaveformEditor::BindCommands()
 		FExecuteAction::CreateSP(TransportController.ToSharedRef(), &FWaveformEditorTransportController::Pause),
 		FCanExecuteAction::CreateSP(TransportController.ToSharedRef(), &FWaveformEditorTransportController::IsPlaying));
 
+	ToolkitCommands->MapAction(
+		Commands.ReturnToStart,
+		FExecuteAction::CreateSP(this, &FWaveformEditor::ResetPlaybackToStart));
+
 	ToolkitCommands->MapAction(
 		Commands.ZoomIn,
 		FExecuteAction::CreateSP(ZoomManager.ToSharedRef(), &FWaveformEditorZoomController::ZoomIn),
diff --git a/Engine/Plugins/Editor/WaveformEditor/Source/WaveformEditor/Private/WaveformEditorCommands.cpp b/Engine/Plugins/Editor/WaveformEditor/Source/WaveformEditor/Private/WaveformEditorCommands.cpp
--- a/Engine/Plugins/Editor/WaveformEditor/Source/WaveformEditor/Private/WaveformEditorCommands.cpp
+++ b/Engine/Plugins/Editor/WaveformEditor/Source/WaveformEditor/Private/WaveformEditorCommands.cpp
@@ -19,6 +19,7 @@ void FWaveformEditorCommands::RegisterCommands()
 	UI_COMMAND(PauseSoundWave, "Pause", "Pause this sound wave", EUserInterfaceActionType::Button, FInputChord());
 	UI_COMMAND(StopSoundWave, "Stop", "Stops playback and rewinds to beginning of file", EUserInterfaceActionType::Button, FInputChord(EKeys::Enter));
 	UI_COMMAND(TogglePlayback, "TogglePlayback", "Toggles between play and pause", EUserInterfaceActionType::Button, FInputChord(EKeys::SpaceBar));
+	UI_COMMAND(ReturnToStart, "Return To Start", "Sets playhead to start of waveform", EUserInterfaceActionType::Button, FInputChord(EKeys::SpaceBar, EModifierKey::Shift));
 
 	UI_COMMAND(ZoomIn, "Zoom In", "Zooms into the sound wave", EUserInterfaceActionType::Button, FInputChord());
 	UI_COMMAND(ZoomOut, "Zoom Out", "Zooms out from the sound wave", EUserInterfaceActionType::Button, FInputChord());
diff --git a/Engine/Plugins/Editor/WaveformEditor/Source/WaveformEditor/Public/WaveformEditor.h b/Engine/Plugins/Editor/WaveformEditor/Source/WaveformEditor/Public/WaveformEditor.h
--- a/Engine/Plugins/Editor/WaveformEditor/Source/WaveformEditor/Public/WaveformEditor.h
+++ b/Engine/Plugins/Editor/WaveformEditor/Source/WaveformEditor/Public/WaveformEditor.h
@@ -62,6 +62,7 @@ private:
 	UE_API bool SetUpAssetReimport();
 	UE_API void ExecuteReimport();
 	UE_API void ExecuteOverwriteTransformations();
+	UE_API void ResetPlaybackToStart();
 
 	/**	Sets the wave editor layout */
 	UE_API const TSharedRef<FTabManager::FLayout> SetupStandaloneLayout();
diff --git a/Engine/Plugins/Editor/WaveformEditor/Source/WaveformEditor/Public/WaveformEditorCommands.h b/Engine/Plugins/Editor/WaveformEditor/Source/WaveformEditor/Public/WaveformEditorCommands.h
--- a/Engine/Plugins/Editor/WaveformEditor/Source/WaveformEditor/Public/WaveformEditorCommands.h
+++ b/Engine/Plugins/Editor/WaveformEditor/Source/WaveformEditor/Public/WaveformEditorCommands.h
@@ -19,6 +19,7 @@ public:
 	TSharedPtr<FUICommandInfo> PauseSoundWave;
 	TSharedPtr<FUICommandInfo> StopSoundWave;
 	TSharedPtr<FUICommandInfo> TogglePlayback;
+	TSharedPtr<FUICommandInfo> ReturnToStart;
 
 	TSharedPtr<FUICommandInfo> ZoomIn;
 	TSharedPtr<FUICommandInfo> ZoomOut;
