Reinstate Waveform editor as default SoundWave editor

diff --git a/Engine/Plugins/Editor/WaveformEditor/Source/WaveformEditor/Private/WaveformEditorInstantiator.cpp b/Engine/Plugins/Editor/WaveformEditor/Source/WaveformEditor/Private/WaveformEditorInstantiator.cpp
--- a/Engine/Plugins/Editor/WaveformEditor/Source/WaveformEditor/Private/WaveformEditorInstantiator.cpp
+++ b/Engine/Plugins/Editor/WaveformEditor/Source/WaveformEditor/Private/WaveformEditorInstantiator.cpp
@@ -29,26 +29,6 @@ namespace WaveEditorInstantiatorPrivate
 
 void FWaveformEditorInstantiator::ExtendContentBrowserSelectionMenu()
 {
-	const UAssetDefinition* AssetDefinition = UAssetDefinitionRegistry::Get()->GetAssetDefinitionForClass(USoundWave::StaticClass());
-	FToolMenuSection* Section = CastChecked<UAssetDefinition_SoundAssetBase>(AssetDefinition)->FindSoundContextMenuSection("Sound");
-	check(Section);
-
-	Section->AddDynamicEntry("SoundWaveEditing_CreateWaveformEditor", FNewToolMenuSectionDelegate::CreateLambda([this](FToolMenuSection& InSection)
-	{
-		if (const UContentBrowserAssetContextMenuContext* Context = InSection.FindContext<UContentBrowserAssetContextMenuContext>())
-		{
-			if (Algo::AnyOf(Context->SelectedAssets, WaveEditorInstantiatorPrivate::FilterUnwantedAssets))
-			{
-				const TAttribute<FText> Label = LOCTEXT("SoundWave_WaveformEditor", "Edit Waveform");
-				const TAttribute<FText> ToolTip = LOCTEXT("SoundWave_WaveformEditor_Tooltip", "Open waveform editor");
-				const FSlateIcon Icon = FSlateIcon();
-				const FToolMenuExecuteAction UIAction = FToolMenuExecuteAction::CreateSP(this, &FWaveformEditorInstantiator::ExecuteCreateWaveformEditor);
-
-				InSection.AddMenuEntry("SoundWave_CreateWaveformEditor", Label, ToolTip, Icon, UIAction);
-			}
-		}
-	}));
-	
 	UToolMenu* Menu = UE::ContentBrowser::ExtendToolMenu_AssetContextMenu(USoundWave::StaticClass());
 	check(Menu);
 	FToolMenuSection& ImportSection = Menu->FindOrAddSection("ImportedAssetActions");
@@ -80,15 +60,6 @@ void FWaveformEditorInstantiator::RegisterAsSoundwaveEditor()
 	}
 }
 
-void FWaveformEditorInstantiator::ExecuteCreateWaveformEditor(const FToolMenuContext& MenuContext)
-{
-	if (const UContentBrowserAssetContextMenuContext* Context = UContentBrowserAssetContextMenuContext::FindContextWithAssets(MenuContext))
-	{
-		TArray<USoundWave*> SoundWavesToEdit = Context->LoadSelectedObjectsIf<USoundWave>(WaveEditorInstantiatorPrivate::FilterUnwantedAssets);
-		CreateWaveformEditor(SoundWavesToEdit);
-	}
-}
-
 void FWaveformEditorInstantiator::ExecuteReImportAndOverwriteTransformations(const FToolMenuContext& MenuContext)
 {
 	if (const UContentBrowserAssetContextMenuContext* Context = UContentBrowserAssetContextMenuContext::FindContextWithAssets(MenuContext))
diff --git a/Engine/Plugins/Editor/WaveformEditor/Source/WaveformEditor/Private/WaveformEditorInstantiator.h b/Engine/Plugins/Editor/WaveformEditor/Source/WaveformEditor/Private/WaveformEditorInstantiator.h
--- a/Engine/Plugins/Editor/WaveformEditor/Source/WaveformEditor/Private/WaveformEditorInstantiator.h
+++ b/Engine/Plugins/Editor/WaveformEditor/Source/WaveformEditor/Private/WaveformEditorInstantiator.h
@@ -19,7 +19,6 @@ public:
 
 private:
 	virtual void CreateWaveformEditor(TArray<USoundWave*> SoundWavesToEdit) override;
-	void ExecuteCreateWaveformEditor(const FToolMenuContext& MenuContext);
 	void ExecuteReImportAndOverwriteTransformations(const FToolMenuContext& MenuContext);
 	bool CanSoundWaveBeOpenedInEditor(const USoundWave* SoundWaveToEdit);
 	void DisplayErrorDialog(const FText& ErrorMessage) const;
diff --git a/Engine/Plugins/Editor/WaveformEditor/Source/WaveformEditor/Private/WaveformEditorModule.cpp b/Engine/Plugins/Editor/WaveformEditor/Source/WaveformEditor/Private/WaveformEditorModule.cpp
--- a/Engine/Plugins/Editor/WaveformEditor/Source/WaveformEditor/Private/WaveformEditorModule.cpp
+++ b/Engine/Plugins/Editor/WaveformEditor/Source/WaveformEditor/Private/WaveformEditorModule.cpp
@@ -20,8 +20,7 @@ void FWaveformEditorModule::StartupModule()
 
 	WaveformEditorInstantiator = MakeShared<FWaveformEditorInstantiator>();
 	RegisterContentBrowserExtensions(WaveformEditorInstantiator.Get());
-
-	// TODO: Re register WaveformEditor as default soundwave editor with: WaveformEditorInstantiator->RegisterAsSoundwaveEditor();
+	WaveformEditorInstantiator->RegisterAsSoundwaveEditor();
 }
 
 void FWaveformEditorModule::ShutdownModule()
