[WaveformEditor] WaveformEditor Transformations log category now can be used in other modules. Reimplementing after being removed by build pipeline fix

diff --git a/Engine/Plugins/Editor/WaveformEditor/Source/WaveformTransformations/Public/WaveformTransformationLog.h b/Engine/Plugins/Editor/WaveformEditor/Source/WaveformTransformations/Public/WaveformTransformationLog.h
--- a/Engine/Plugins/Editor/WaveformEditor/Source/WaveformTransformations/Public/WaveformTransformationLog.h
+++ b/Engine/Plugins/Editor/WaveformEditor/Source/WaveformTransformations/Public/WaveformTransformationLog.h
@@ -4,4 +4,8 @@
 
 #include "CoreMinimal.h"
 
-DECLARE_LOG_CATEGORY_EXTERN(LogWaveformTransformation, Log, All);
\ No newline at end of file
+#define UE_API WAVEFORMTRANSFORMATIONS_API
+
+UE_API DECLARE_LOG_CATEGORY_EXTERN(LogWaveformTransformation, Log, All);
+
+#undef UE_API
\ No newline at end of file
diff --git a/Engine/Plugins/Editor/WaveformEditor/Source/WaveformTransformationsWidgets/Private/WaveformTransformationTrimFadeRenderer.cpp b/Engine/Plugins/Editor/WaveformEditor/Source/WaveformTransformationsWidgets/Private/WaveformTransformationTrimFadeRenderer.cpp
--- a/Engine/Plugins/Editor/WaveformEditor/Source/WaveformTransformationsWidgets/Private/WaveformTransformationTrimFadeRenderer.cpp
+++ b/Engine/Plugins/Editor/WaveformEditor/Source/WaveformTransformationsWidgets/Private/WaveformTransformationTrimFadeRenderer.cpp
@@ -10,6 +10,7 @@
 #include "PropertyHandle.h"
 #include "Rendering/DrawElements.h"
 #include "Sound/SoundWave.h"
+#include "WaveformTransformationLog.h"
 #include "Widgets/Images/SImage.h"
 #include "Widgets/Input/SButton.h"
 #include "Widgets/SWidget.h"
@@ -292,7 +293,7 @@ FReply FWaveformTransformationTrimFadeRenderer::OnMouseButtonDown(SWidget& Owner
 	{
 		if (BeginTransaction(TEXT("PropertyEditor"), FText::Format(NSLOCTEXT("PropertyEditor", "SetNumericPropertyTransaction", "Edit {0}"), GetPropertyEditedByCurrentInteraction()), StrongTrimFade.Get()) == INDEX_NONE)
 		{
-			//TODO: Log a warning
+			UE_LOG(LogWaveformTransformation, Warning, TEXT("Begin TrimFadeTransformation OnMouseButtonDown Transaction Failed"));
 		}
 		StrongTrimFade->Modify();
 		return FReply::Handled().CaptureMouse(OwnerWidget.AsShared()).PreventThrottling();
