Waveform Editor: loop delete and conversion

diff --git a/Engine/Plugins/Editor/WaveformEditor/Source/WaveformEditor/Private/WaveformEditorCommands.cpp b/Engine/Plugins/Editor/WaveformEditor/Source/WaveformEditor/Private/WaveformEditorCommands.cpp
--- a/Engine/Plugins/Editor/WaveformEditor/Source/WaveformEditor/Private/WaveformEditorCommands.cpp
+++ b/Engine/Plugins/Editor/WaveformEditor/Source/WaveformEditor/Private/WaveformEditorCommands.cpp
@@ -57,7 +57,7 @@ void FWaveformEditorCommands::RegisterCommands()
 
 	UI_COMMAND(CreateMarker, "Create Marker", "Toggle adding Marker Cue Transformation onto the sound wave", EUserInterfaceActionType::Button, FInputChord());
 	UI_COMMAND(CreateLoopRegion, "Create Loop Region", "Toggle adding Marker Loop Region Transformation onto the sound wave", EUserInterfaceActionType::Button, FInputChord());
-	UI_COMMAND(DeleteMarker, "Delete Marker Or Loop Region", "Delete a selected Marker Cue or Loop Region from the sound wave", EUserInterfaceActionType::Button, FInputChord());
+	UI_COMMAND(DeleteMarker, "Delete Marker Or Loop Region", "Delete a selected Marker Cue or Loop Region from the sound wave", EUserInterfaceActionType::Button, FInputChord(EKeys::Delete));
 
 	UI_COMMAND(SkipToNextMarker, "Skip To Next Marker", "Skips the playhead to the next marker", EUserInterfaceActionType::Button, FInputChord(EKeys::G));
 }
diff --git a/Engine/Plugins/Editor/WaveformEditor/Source/WaveformTransformations/Private/WaveformTransformationMarkers.cpp b/Engine/Plugins/Editor/WaveformEditor/Source/WaveformTransformations/Private/WaveformTransformationMarkers.cpp
--- a/Engine/Plugins/Editor/WaveformEditor/Source/WaveformTransformations/Private/WaveformTransformationMarkers.cpp
+++ b/Engine/Plugins/Editor/WaveformEditor/Source/WaveformTransformations/Private/WaveformTransformationMarkers.cpp
@@ -669,10 +669,10 @@ void UWaveCueArray::Reset()
 }
 #endif
 
-void UWaveCueArray::EnableLoopRegion(FSoundWaveCuePoint* OutSoundWaveCue)
+void UWaveCueArray::EnableLoopRegion(FSoundWaveCuePoint* OutSoundWaveCue, bool bSetLoopRegion)
 {
 	check(OutSoundWaveCue != nullptr);
-	OutSoundWaveCue->SetLoopRegion(true);
+	OutSoundWaveCue->SetLoopRegion(bSetLoopRegion);
 }
 
 #undef LOCTEXT_NAMESPACE
\ No newline at end of file
diff --git a/Engine/Plugins/Editor/WaveformEditor/Source/WaveformTransformations/Public/WaveformTransformationMarkers.h b/Engine/Plugins/Editor/WaveformEditor/Source/WaveformTransformations/Public/WaveformTransformationMarkers.h
--- a/Engine/Plugins/Editor/WaveformEditor/Source/WaveformTransformations/Public/WaveformTransformationMarkers.h
+++ b/Engine/Plugins/Editor/WaveformEditor/Source/WaveformTransformations/Public/WaveformTransformationMarkers.h
@@ -37,7 +37,7 @@ public:
 	//Uninitialize and empty Marker array
 	UE_API void Reset();
 
-	UE_API void EnableLoopRegion(FSoundWaveCuePoint* OutSoundWaveCue);
+	UE_API void EnableLoopRegion(FSoundWaveCuePoint* OutSoundWaveCue, bool bSetLoopRegion = true);
 
 	UPROPERTY(EditAnywhere, Category = "Markers")
 	TArray<FSoundWaveCuePoint> CuesAndLoops;
diff --git a/Engine/Plugins/Editor/WaveformEditor/Source/WaveformTransformationsWidgets/Private/WaveformTransformationMarkersRenderer.cpp b/Engine/Plugins/Editor/WaveformEditor/Source/WaveformTransformationsWidgets/Private/WaveformTransformationMarkersRenderer.cpp
--- a/Engine/Plugins/Editor/WaveformEditor/Source/WaveformTransformationsWidgets/Private/WaveformTransformationMarkersRenderer.cpp
+++ b/Engine/Plugins/Editor/WaveformEditor/Source/WaveformTransformationsWidgets/Private/WaveformTransformationMarkersRenderer.cpp
@@ -713,7 +713,7 @@ FWaveformTransformationMarkerRenderer::EMarkerInteractionType FWaveformTransform
 	check(StrongMarkersTransformation);
 	if (MouseButton == EKeys::RightMouseButton)
 	{
-		if (SetMarkerInInteractionRange(InLocalCursorPosition, WidgetGeometry) && MarkerInInteractionRange.GetMarker() != nullptr && MarkerInInteractionRange.GetMarker()->IsLoopRegion() == false)
+		if (SetMarkerInInteractionRange(InLocalCursorPosition, WidgetGeometry) && MarkerInInteractionRange.GetMarker() != nullptr)
 		{
 			return EMarkerInteractionType::RightClickMarker;
 		}
@@ -869,8 +869,8 @@ void FWaveformTransformationMarkerRenderer::ShowConvertToLoopAtMarker(const FGeo
 		.AutoHeight()
 		[
 			SNew(SButton)
-				.Text(LOCTEXT("WaveformTransformationMarkerRendererConvertLoop", "Convert To Loop Region"))
-				.OnClicked_Lambda([InCreateLoopMenuWindow = CreateLoopMenuWindow, WeakMarkersTransformation, CueToModify, MinLoopSize = MarkerInInteractionRange.GetMinLoopSize()]() -> FReply
+				.Text(LOCTEXT("WaveformTransformationMarkerRendererToggleLoop", "Toggle Loop Region"))
+				.OnClicked_Lambda([this, InCreateLoopMenuWindow = CreateLoopMenuWindow, WeakMarkersTransformation, CueToModify, MinLoopSize = MarkerInInteractionRange.GetMinLoopSize()]() -> FReply
 					{
 						if (TStrongObjectPtr<UWaveformTransformationMarkers> LockedMarkers = WeakMarkersTransformation.Pin())
 						{
@@ -880,8 +880,16 @@ void FWaveformTransformationMarkerRenderer::ShowConvertToLoopAtMarker(const FGeo
 							}
 
 							LockedMarkers->Markers->Modify();
-							LockedMarkers->Markers->EnableLoopRegion(CueToModify);
-							CueToModify->FrameLength = MinLoopSize;
+
+							if (CueToModify->FrameLength < MinLoopSize)
+							{
+								CueToModify->FrameLength = MinLoopSize;
+							}
+
+							LockedMarkers->Markers->EnableLoopRegion(CueToModify, !CueToModify->IsLoopRegion());
+
+							SetActiveLoopRegion(LockedMarkers->Markers->SelectedCue);
+							StrongMarkersTransformation->OnTransformationRenderChanged.ExecuteIfBound();
 						}
 
 						if (GEditor && GEditor->Trans)
@@ -898,6 +906,43 @@ void FWaveformTransformationMarkerRenderer::ShowConvertToLoopAtMarker(const FGeo
 
 		];
 
+	MenuContent->AddSlot()
+		.AutoHeight()
+		[
+			SNew(SButton)
+				.Text(LOCTEXT("WaveformTransformationMarkerRendererDelete", "Delete"))
+				.OnClicked_Lambda([this, CueToModify]() -> FReply
+					{
+						BeginTransaction(TEXT("PropertyEditor"), LOCTEXT("Delete", "Delete"), nullptr);
+
+						MarkersArray->Modify();
+						if (CueToModify->IsLoopRegion())
+						{
+							SetIsPreviewingLoopRegion(0.0, -1.0, false);
+						}
+
+						for (int32 i = 0; i < MarkersArray->CuesAndLoops.Num(); i++)
+						{
+							if (CueToModify->CuePointID == MarkersArray->CuesAndLoops[i].CuePointID)
+							{
+								MarkersArray->CuesAndLoops.RemoveAt(i);
+							}
+						}
+
+						MarkersArray->SelectedCue = INDEX_NONE;
+
+						StrongMarkersTransformation->OnTransformationRenderChanged.ExecuteIfBound();
+
+						EndTransaction();
+
+						if (TSharedPtr<SWindow> LockedFadeModeMenuWindow = CreateLoopMenuWindow.Pin())
+						{
+							FSlateApplication::Get().RequestDestroyWindow(LockedFadeModeMenuWindow.ToSharedRef());
+						}
+						return FReply::Handled();
+					})
+		];
+
 	FSlateApplication::Get().AddWindow(MenuWindow);
 
 	FTSTicker::GetCoreTicker().AddTicker(FTickerDelegate::CreateLambda([this, InCreateLoopMenuWindow = CreateLoopMenuWindow, LocalCursorPosition, LocalWindowMaxPosition](float DeltaTime)
