[Waveform Editor] - Able to drag Isolated region of a Loop Region with a Start Trim

diff --git a/Engine/Plugins/Editor/WaveformEditor/Source/WaveformEditor/Private/WaveformEditor.cpp b/Engine/Plugins/Editor/WaveformEditor/Source/WaveformEditor/Private/WaveformEditor.cpp
--- a/Engine/Plugins/Editor/WaveformEditor/Source/WaveformEditor/Private/WaveformEditor.cpp
+++ b/Engine/Plugins/Editor/WaveformEditor/Source/WaveformEditor/Private/WaveformEditor.cpp
@@ -645,7 +645,9 @@ void FWaveformEditor::UpdateTransformations()
 
 void FWaveformEditor::UpdateRenderElements()
 {
+	check(SoundWave);
 	WaveformView.DataProvider->UpdateRenderElements();
+	TransformationChainConfig = SoundWave->UpdateTransformations();
 }
 
 void FWaveformEditor::ModifyMarkerLoopRegion(ELoopModificationControls Modification)
diff --git a/Engine/Plugins/Editor/WaveformEditor/Source/WaveformTransformations/Private/WaveformTransformationMarkers.cpp b/Engine/Plugins/Editor/WaveformEditor/Source/WaveformTransformations/Private/WaveformTransformationMarkers.cpp
--- a/Engine/Plugins/Editor/WaveformEditor/Source/WaveformTransformations/Private/WaveformTransformationMarkers.cpp
+++ b/Engine/Plugins/Editor/WaveformEditor/Source/WaveformTransformations/Private/WaveformTransformationMarkers.cpp
@@ -452,23 +452,20 @@ bool UWaveformTransformationMarkers::AdjustLoopPreviewIfNotAligned()
 	const int64 CuePointPosDif = FMath::Abs(CuePoint->FramePosition - StartLoopFramePos - StartFrameOffset);
 	const int64 CuePointEndPosDif = FMath::Abs((CuePoint->FramePosition + CuePoint->FrameLength) - EndLoopFramePos - StartFrameOffset);
 
-	if (StartFrameOffset > CuePoint->FramePosition && StartLoopTime == 0.0)
-	{
-		return false;
-	}
-
 	// Converting from time to frame is not perfectly accurate, lets be +- 2 frames
 	if (CuePointPosDif > 2)
 	{
+		const double PrevStartLoopTime = StartLoopTime;
 		StartLoopTime = static_cast<float>(FMath::Max(0, CuePoint->FramePosition - StartFrameOffset)) / SampleRate;
-		bPreviewChanged = true;
+		bPreviewChanged = (StartLoopTime != PrevStartLoopTime);
 	}
 
 	// Converting from time to frame is not perfectly accurate, lets be +- 2 frames
 	if (CuePointEndPosDif > 2)
 	{
+		const double PrevEndLoopTime = EndLoopTime;
 		EndLoopTime = static_cast<float>((CuePoint->FramePosition + CuePoint->FrameLength) - StartFrameOffset) / SampleRate;
-		bPreviewChanged = true;
+		bPreviewChanged |= (EndLoopTime != PrevEndLoopTime);
 	}
 
 	return bPreviewChanged;
