Waveform Editor: Streamline trimming workflow

diff --git a/Engine/Plugins/Editor/WaveformEditor/Source/WaveformEditor/Private/WaveformEditor.cpp b/Engine/Plugins/Editor/WaveformEditor/Source/WaveformEditor/Private/WaveformEditor.cpp
--- a/Engine/Plugins/Editor/WaveformEditor/Source/WaveformEditor/Private/WaveformEditor.cpp
+++ b/Engine/Plugins/Editor/WaveformEditor/Source/WaveformEditor/Private/WaveformEditor.cpp
@@ -41,6 +41,7 @@ const FName FWaveformEditor::TransformationsTabId("WaveformEditor_Transformation
 const FName FWaveformEditor::WaveformDisplayTabId("WaveformEditor_Display");
 const FName FWaveformEditor::EditorName("Waveform Editor");
 const FName FWaveformEditor::ToolkitFName("WaveformEditor");
+const double FWaveformEditor::EndTimeInvalid = -1.0;
 
 bool FWaveformEditor::Init(const EToolkitMode::Type Mode, const TSharedPtr<IToolkitHost>& InitToolkitHost, USoundWave* SoundWaveToEdit)
 {
@@ -65,6 +66,16 @@ bool FWaveformEditor::Init(const EToolkitMode::Type Mode, const TSharedPtr<ITool
 		TransformationChainConfig = SoundWave->UpdateTransformations();
 	}
 
+	// Add a trimfade transformation for user to opt to use, but will be removed from transformations on save if values are zero
+	// This does not trigger a modify on the soundwave to avoid marking dirty every time a user opens the waveform, but if they change anything on the trimfade that automatically will modify and mark dirty to save it
+	TObjectPtr<UWaveformTransformationTrimFade> TrimFadeTransformation = GetOrAddTrimFadeTransformation(false);
+	check(TrimFadeTransformation != nullptr);
+
+	if (FMath::IsNearlyEqual(TrimFadeTransformation->EndTime, EndTimeInvalid))
+	{
+		TrimFadeTransformation->EndTime = SoundWave->Duration;
+	}
+
 	bool bIsInitialized = true;
 	
 	bIsInitialized &= CreateDetailsViews();
@@ -153,8 +164,6 @@ bool FWaveformEditor::Init(const EToolkitMode::Type Mode, const TSharedPtr<ITool
 			bMarkerTransformationIsActive = false;
 		}
 
-		TObjectPtr<UWaveformTransformationTrimFade> TrimFadeTransformation = GetTrimFadeTransformation();
-
 		if (TrimFadeTransformation != nullptr)
 		{
 			if (!TrimFadeTransformation->OnTransformationChanged.IsBound())
@@ -244,6 +253,22 @@ FWaveformEditor::~FWaveformEditor()
 	{
 		FReimportManager::Instance()->OnPostReimport().RemoveAll(this);
 	}
+
+	// If the user has not used the trimfade then remove the transformation to avoid unecessary transformations in the chain
+	check(TrimFadeTransformation != nullptr);
+	check(SoundWave != nullptr);
+
+	// Account for a user trimming a single frame by having a 1/2 frame error tolerance
+	uint32 ImportedSampleRate = SoundWave->GetImportedSampleRate();
+	check(ImportedSampleRate > 0);
+
+	const double ErrorTolerance = (1.0f / ImportedSampleRate) * 0.5f;
+
+	if (FMath::IsNearlyEqual(TrimFadeTransformation->StartTime, 0, ErrorTolerance) && FMath::IsNearlyEqual(TrimFadeTransformation->EndTime, SoundWave->Duration, ErrorTolerance)
+		&& FMath::IsNearlyEqual(TrimFadeTransformation->StartFadeTime, 0, ErrorTolerance) && FMath::IsNearlyEqual(TrimFadeTransformation->EndFadeTime, 0, ErrorTolerance))
+	{
+		SoundWave->Transformations.Remove(TrimFadeTransformation);
+	}
 }
 
 void FWaveformEditor::AddDefaultTransformations()
@@ -307,7 +332,7 @@ void FWaveformEditor::NotifyPostTransformationChange(const EPropertyChangeType::
 	TransformationChainConfig = SoundWave->UpdateTransformations();
 }
 
-TObjectPtr<UWaveformTransformationTrimFade> FWaveformEditor::GetOrAddTrimFadeTransformation()
+TObjectPtr<UWaveformTransformationTrimFade> FWaveformEditor::GetOrAddTrimFadeTransformation(const bool bMarkModify)
 {
 	check(SoundWave);
 
@@ -324,10 +349,16 @@ TObjectPtr<UWaveformTransformationTrimFade> FWaveformEditor::GetOrAddTrimFadeTra
 
 		TrimFadeTransformation = NewObject<UWaveformTransformationTrimFade>(SoundWave, UWaveformTransformationTrimFade::StaticClass(), NAME_None, MaskedOuterFlags);
 
-		SoundWave->Modify();
-		SoundWave->Transformations.Add(TrimFadeTransformation);
-
-		NotifyPostTransformationChange();
+		if (bMarkModify)
+		{
+			SoundWave->Modify();
+			SoundWave->Transformations.Add(TrimFadeTransformation);
+			NotifyPostTransformationChange();
+		}
+		else
+		{
+			SoundWave->Transformations.Add(TrimFadeTransformation);
+		}
 	}
 	return TrimFadeTransformation;
 }
diff --git a/Engine/Plugins/Editor/WaveformEditor/Source/WaveformEditor/Public/WaveformEditor.h b/Engine/Plugins/Editor/WaveformEditor/Source/WaveformEditor/Public/WaveformEditor.h
--- a/Engine/Plugins/Editor/WaveformEditor/Source/WaveformEditor/Public/WaveformEditor.h
+++ b/Engine/Plugins/Editor/WaveformEditor/Source/WaveformEditor/Public/WaveformEditor.h
@@ -54,6 +54,8 @@ public:
 
 	UE_API void ExportWaveform(bool bPromptUserToSave);
 
+	static UE_API const double EndTimeInvalid;
+
 private:
 	UE_API bool InitializeAudioComponent();
 	UE_API bool CreateTransportController();
@@ -107,9 +109,9 @@ private:
 
 	/** Transformation Functions*/
 	UE_API void NotifyPostTransformationChange(const EPropertyChangeType::Type& PropertyChangeType = EPropertyChangeType::ArrayAdd);
-	UE_API TObjectPtr<UWaveformTransformationTrimFade> GetOrAddTrimFadeTransformation();
 	// Relies on there being only one TrimFade Transformation on the soundwave
 	UE_API TObjectPtr<UWaveformTransformationTrimFade> GetTrimFadeTransformation() const;
+	UE_API TObjectPtr<UWaveformTransformationTrimFade> GetOrAddTrimFadeTransformation(const bool bMarkModify = true);
 
 	UE_API TObjectPtr<UWaveformTransformationMarkers> GetOrAddMarkerTransformation();
 	// Relies on there being only one Marker Transformation on the soundwave
diff --git a/Engine/Plugins/Editor/WaveformEditor/Source/WaveformTransformationsWidgets/Private/WaveformTransformationTrimFadeRenderer.cpp b/Engine/Plugins/Editor/WaveformEditor/Source/WaveformTransformationsWidgets/Private/WaveformTransformationTrimFadeRenderer.cpp
--- a/Engine/Plugins/Editor/WaveformEditor/Source/WaveformTransformationsWidgets/Private/WaveformTransformationTrimFadeRenderer.cpp
+++ b/Engine/Plugins/Editor/WaveformEditor/Source/WaveformTransformationsWidgets/Private/WaveformTransformationTrimFadeRenderer.cpp
@@ -35,40 +35,30 @@ int32 FWaveformTransformationTrimFadeRenderer::DrawTrimHandles(const FGeometry&
 	const bool bRenderLowerBound = StartTimeHandleX >= 0.f;
 	const bool bRenderUpperBound = EndTimeHandleX <= AllottedGeometry.Size.X;
 
-	TArray<FVector2D> LinePoints;
-	LinePoints.SetNumUninitialized(2);
-
-	if (bRenderLowerBound)
-	{
-		LinePoints[0] = FVector2D(StartTimeHandleX, 0.f);
-		LinePoints[1] = FVector2D(StartTimeHandleX, AllottedGeometry.Size.Y);
-
-		FSlateDrawElement::MakeLines(
-			OutDrawElements,
-			++LayerId,
-			AllottedGeometry.ToPaintGeometry(),
-			LinePoints,
-			ESlateDrawEffect::None,
-			FLinearColor::Green,
-			false
-		);
-	}
-
-	if (bRenderUpperBound)
-	{
-		LinePoints[0] = FVector2D(EndTimeHandleX, 0.f);
-		LinePoints[1] = FVector2D(EndTimeHandleX, AllottedGeometry.Size.Y);
-
-		FSlateDrawElement::MakeLines(
-			OutDrawElements,
-			++LayerId,
-			AllottedGeometry.ToPaintGeometry(),
-			LinePoints,
-			ESlateDrawEffect::None,
-			FLinearColor::Red,
-			false
-		);
-	}
+	float HandleWidth = InteractionHandleWidth;
+
+	const FSlateRoundedBoxBrush RoundedBoxBrush(FLinearColor::Green, 2);
+	FPaintGeometry HandleGeometry = AllottedGeometry.ToPaintGeometry(FVector2f(HandleWidth, HandleWidth), FSlateLayoutTransform(FVector2f(StartTimeHandleX, AllottedGeometry.Size.Y - HandleWidth)));
+
+	FSlateDrawElement::MakeBox(
+		OutDrawElements,
+		++LayerId,
+		HandleGeometry,
+		&RoundedBoxBrush,
+		ESlateDrawEffect::None,
+		FLinearColor::Green
+	);
+
+	HandleGeometry = AllottedGeometry.ToPaintGeometry(FVector2f(HandleWidth, HandleWidth), FSlateLayoutTransform(FVector2f(EndTimeHandleX - HandleWidth, AllottedGeometry.Size.Y - HandleWidth)));
+
+	FSlateDrawElement::MakeBox(
+		OutDrawElements,
+		++LayerId,
+		HandleGeometry,
+		&RoundedBoxBrush,
+		ESlateDrawEffect::None,
+		FLinearColor::Red
+	);
 
 	return LayerId;
 }
@@ -103,6 +93,19 @@ int32 FWaveformTransformationTrimFadeRenderer::DrawFadeCurves(const FGeometry& A
 			FadeColor
 		);
 	}
+	else
+	{
+		FPaintGeometry HandleGeometry = AllottedGeometry.ToPaintGeometry(FVector2f(HandleWidth, HandleHeight), FSlateLayoutTransform(FVector2f(StartTimeHandleX, 0)));
+
+		FSlateDrawElement::MakeBox(
+			OutDrawElements,
+			++LayerId,
+			HandleGeometry,
+			&RoundedBoxBrush,
+			ESlateDrawEffect::None,
+			FadeColor
+		);
+	}
 
 	if (FadeOutCurvePoints.Num() > 0)
 	{
@@ -127,6 +130,19 @@ int32 FWaveformTransformationTrimFadeRenderer::DrawFadeCurves(const FGeometry& A
 			FadeColor
 		);
 	}
+	else
+	{
+		FPaintGeometry HandleGeometry = AllottedGeometry.ToPaintGeometry(FVector2f(HandleWidth, HandleHeight), FSlateLayoutTransform(FVector2f(EndTimeHandleX - HandleWidth, 0)));
+
+		FSlateDrawElement::MakeBox(
+			OutDrawElements,
+			++LayerId,
+			HandleGeometry,
+			&RoundedBoxBrush,
+			ESlateDrawEffect::None,
+			FadeColor
+		);
+	}
 
 	return LayerId;
 }
@@ -186,7 +202,9 @@ FCursorReply FWaveformTransformationTrimFadeRenderer::OnCursorQuery(const FGeome
 		return FCursorReply::Cursor(EMouseCursor::GrabHand);
 	}
 
-	if (TrimFadeInteractionType == ETrimFadeInteractionType::ScrubbingLeftHandle || TrimFadeInteractionType == ETrimFadeInteractionType::ScrubbingRightHandle || StartTimeInteractionXRange.Contains(LocalCursorPosition.X) || EndTimeInteractionXRange.Contains(LocalCursorPosition.X))
+	if ((TrimFadeInteractionType == ETrimFadeInteractionType::ScrubbingLeftHandle || TrimFadeInteractionType == ETrimFadeInteractionType::ScrubbingRightHandle) 
+		|| (EndTimeInteractionXRange.Contains(LocalCursorPosition.X) && EndTimeInteractionYRange.Contains(LocalCursorPosition.Y))
+		|| (StartTimeInteractionXRange.Contains(LocalCursorPosition.X) && StartTimeInteractionYRange.Contains(LocalCursorPosition.Y)))
 	{
 		return FCursorReply::Cursor(EMouseCursor::ResizeLeftRight);
 	}
@@ -265,19 +283,25 @@ void FWaveformTransformationTrimFadeRenderer::Tick(const FGeometry& AllottedGeom
 	EndTimeHandleX = EndFrame * PixelsPerFrame;
 
 	GenerateFadeCurves(AllottedGeometry);
-	UpdateInteractionRange();
+	UpdateInteractionRange(AllottedGeometry);
 }
 
-void FWaveformTransformationTrimFadeRenderer::UpdateInteractionRange()
+void FWaveformTransformationTrimFadeRenderer::UpdateInteractionRange(const FGeometry& AllottedGeometry)
 {
-	StartTimeInteractionXRange.SetLowerBoundValue(StartTimeHandleX - InteractionPixelXDelta);
-	StartTimeInteractionXRange.SetUpperBoundValue(StartTimeHandleX + InteractionPixelXDelta);
-	EndTimeInteractionXRange.SetLowerBoundValue(EndTimeHandleX - InteractionPixelXDelta);
-	EndTimeInteractionXRange.SetUpperBoundValue(EndTimeHandleX + InteractionPixelXDelta);
-	FadeInInteractionXRange.SetLowerBoundValue(FadeInEndX - InteractionPixelXDelta);
-	FadeInInteractionXRange.SetUpperBoundValue(FadeInEndX + InteractionPixelXDelta);
-	FadeOutInteractionXRange.SetLowerBoundValue(FadeOutStartX - InteractionPixelXDelta);
-	FadeOutInteractionXRange.SetUpperBoundValue(FadeOutStartX + InteractionPixelXDelta);
+	StartTimeInteractionXRange.SetLowerBoundValue(StartTimeHandleX - InteractionHandleWidth);
+	StartTimeInteractionXRange.SetUpperBoundValue(StartTimeHandleX + InteractionHandleWidth);
+	StartTimeInteractionYRange.SetLowerBoundValue(AllottedGeometry.GetLocalSize().Y - InteractionHandleWidth);
+	StartTimeInteractionYRange.SetUpperBoundValue(AllottedGeometry.GetLocalSize().Y);
+
+	EndTimeInteractionXRange.SetLowerBoundValue(EndTimeHandleX - InteractionHandleWidth);
+	EndTimeInteractionXRange.SetUpperBoundValue(EndTimeHandleX + InteractionHandleWidth);
+	EndTimeInteractionYRange.SetLowerBoundValue(AllottedGeometry.GetLocalSize().Y - InteractionHandleWidth);
+	EndTimeInteractionYRange.SetUpperBoundValue(AllottedGeometry.GetLocalSize().Y);
+
+	FadeInInteractionXRange.SetLowerBoundValue(FadeInEndX - InteractionHandleWidth);
+	FadeInInteractionXRange.SetUpperBoundValue(FadeInEndX + InteractionHandleWidth);
+	FadeOutInteractionXRange.SetLowerBoundValue(FadeOutStartX - InteractionHandleWidth);
+	FadeOutInteractionXRange.SetUpperBoundValue(FadeOutStartX + InteractionHandleWidth);
 }
 
 FReply FWaveformTransformationTrimFadeRenderer::OnMouseButtonDown(SWidget& OwnerWidget, const FGeometry& MyGeometry, const FPointerEvent& MouseEvent)
@@ -575,12 +599,12 @@ FWaveformTransformationTrimFadeRenderer::ETrimFadeInteractionType FWaveformTrans
 			return ETrimFadeInteractionType::ScrubbingFadeOut;
 		}
 
-		if (StartTimeInteractionXRange.Contains(InLocalCursorPosition.X))
+		if (StartTimeInteractionXRange.Contains(InLocalCursorPosition.X) && StartTimeInteractionYRange.Contains(InLocalCursorPosition.Y))
 		{
 			return ETrimFadeInteractionType::ScrubbingLeftHandle;
 		}
 
-		if (EndTimeInteractionXRange.Contains(InLocalCursorPosition.X))
+		if (EndTimeInteractionXRange.Contains(InLocalCursorPosition.X) && EndTimeInteractionYRange.Contains(InLocalCursorPosition.Y))
 		{
 			return ETrimFadeInteractionType::ScrubbingRightHandle;
 		}
diff --git a/Engine/Plugins/Editor/WaveformEditor/Source/WaveformTransformationsWidgets/Private/WaveformTransformationTrimFadeRenderer.h b/Engine/Plugins/Editor/WaveformEditor/Source/WaveformTransformationsWidgets/Private/WaveformTransformationTrimFadeRenderer.h
--- a/Engine/Plugins/Editor/WaveformEditor/Source/WaveformTransformationsWidgets/Private/WaveformTransformationTrimFadeRenderer.h
+++ b/Engine/Plugins/Editor/WaveformEditor/Source/WaveformTransformationsWidgets/Private/WaveformTransformationTrimFadeRenderer.h
@@ -25,7 +25,7 @@ private:
 	int32 DrawTrimHandles(const FGeometry& AllottedGeometry, FSlateWindowElementList& OutDrawElements, int32 LayerId) const;
 	int32 DrawFadeCurves(const FGeometry& AllottedGeometry, FSlateWindowElementList& OutDrawElements, int32 LayerId) const;
 	void GenerateFadeCurves(const FGeometry& AllottedGeometry);
-	void UpdateInteractionRange();
+	void UpdateInteractionRange(const FGeometry& AllottedGeometry);
 
 	bool IsCursorInFadeInInteractionRange(const FVector2D& InLocalCursorPosition, const FGeometry& WidgetGeometry) const;
 	bool IsCursorInFadeOutInteractionRange(const FVector2D& InLocalCursorPosition, const FGeometry& WidgetGeometry) const;
@@ -44,11 +44,14 @@ private:
 	TArray<FVector2D> FadeOutCurvePoints;
 
 	TRange<float> StartTimeInteractionXRange = TRange<float>::Inclusive(0.f, 0.f);
+	TRange<float> StartTimeInteractionYRange = TRange<float>::Inclusive(0.f, 0.f);
 	TRange<float> EndTimeInteractionXRange = TRange<float>::Inclusive(0.f, 0.f);
+	TRange<float> EndTimeInteractionYRange = TRange<float>::Inclusive(0.f, 0.f);
 	TRange<float> FadeInInteractionXRange = TRange<float>::Inclusive(0.f, 0.f);
 	TRange<float> FadeOutInteractionXRange = TRange<float>::Inclusive(0.f, 0.f);
 
 	double PixelsPerFrame = 0.0;
+	const float InteractionHandleWidth = 20.0f;
 
 	enum class ETrimFadeInteractionType : uint8
 	{
