Waveform Editor: nudge adjustments

diff --git a/Engine/Plugins/Editor/WaveformEditor/Source/WaveformEditor/Private/WaveformEditor.cpp b/Engine/Plugins/Editor/WaveformEditor/Source/WaveformEditor/Private/WaveformEditor.cpp
--- a/Engine/Plugins/Editor/WaveformEditor/Source/WaveformEditor/Private/WaveformEditor.cpp
+++ b/Engine/Plugins/Editor/WaveformEditor/Source/WaveformEditor/Private/WaveformEditor.cpp
@@ -686,8 +686,12 @@ void FWaveformEditor::ModifyMarkerLoopRegion(ELoopModificationControls Modificat
 	check(SoundWave);
 	TObjectPtr<UWaveformTransformationMarkers> MarkersTransformation = GetOrAddMarkerTransformation();
 	check(MarkersTransformation);
+	check(MarkersTransformation->Markers);
 
-	MarkersTransformation->ModifyMarkerLoopRegion(Modification);
+	TArray<float> PCMArray = WaveformView.DataProvider->GetTransformedPCMData();
+	check(PCMArray.Num() > 0);
+
+	MarkersTransformation->ModifyMarkerLoopRegion(Modification, PCMArray);
 }
 
 void FWaveformEditor::CycleMarkerLoopRegion(ELoopModificationControls Modification)
@@ -1446,6 +1450,22 @@ bool FWaveformEditor::BindCommands()
 		Commands.RightBoundsDecrease,
 		FExecuteAction::CreateLambda([this] { ModifyMarkerLoopRegion(ELoopModificationControls::RightHandleDecrement); }));
 
+	ToolkitCommands->MapAction(
+		Commands.LeftBoundsNextZeroCrossing,
+		FExecuteAction::CreateLambda([this] { ModifyMarkerLoopRegion(ELoopModificationControls::LeftHandleNextZeroCrossing); }));
+
+	ToolkitCommands->MapAction(
+		Commands.LeftBoundsPreviousZeroCrossing,
+		FExecuteAction::CreateLambda([this] { ModifyMarkerLoopRegion(ELoopModificationControls::LeftHandlePreviousZeroCrossing); }));
+
+	ToolkitCommands->MapAction(
+		Commands.RightBoundsNextZeroCrossing,
+		FExecuteAction::CreateLambda([this] { ModifyMarkerLoopRegion(ELoopModificationControls::RightHandleNextZeroCrossing); }));
+
+	ToolkitCommands->MapAction(
+		Commands.RightBoundsPreviousZeroCrossing,
+		FExecuteAction::CreateLambda([this] { ModifyMarkerLoopRegion(ELoopModificationControls::RightHandlePreviousZeroCrossing); }));
+
 	ToolkitCommands->MapAction(
 		Commands.BoundsIncrementIncrease,
 		FExecuteAction::CreateLambda([this] { ModifyMarkerLoopRegion(ELoopModificationControls::IncreaseIncrement); }));
diff --git a/Engine/Plugins/Editor/WaveformEditor/Source/WaveformEditor/Private/WaveformEditorCommands.cpp b/Engine/Plugins/Editor/WaveformEditor/Source/WaveformEditor/Private/WaveformEditorCommands.cpp
--- a/Engine/Plugins/Editor/WaveformEditor/Source/WaveformEditor/Private/WaveformEditorCommands.cpp
+++ b/Engine/Plugins/Editor/WaveformEditor/Source/WaveformEditor/Private/WaveformEditorCommands.cpp
@@ -50,6 +50,11 @@ void FWaveformEditorCommands::RegisterCommands()
 	UI_COMMAND(RightBoundsIncrease, "Increase Right Loop Bounds", "Moves the end of the loop region earlier", EUserInterfaceActionType::Button, FInputChord(EKeys::RightBracket));
 	UI_COMMAND(RightBoundsDecrease, "Decrease Right Loop Bounds", "Moves the end of the loop region later", EUserInterfaceActionType::Button, FInputChord(EKeys::LeftBracket));
 
+	UI_COMMAND(LeftBoundsNextZeroCrossing, "Increase Left Loop Bounds To Next Zero Crossing", "Moves the start of the loop region earlier", EUserInterfaceActionType::Button, FInputChord(EModifierKey::Shift, EKeys::Comma));
+	UI_COMMAND(LeftBoundsPreviousZeroCrossing, "Decrease Left Loop Bounds To Previous Zero Crossing", "Moves the start of the loop region later", EUserInterfaceActionType::Button, FInputChord(EModifierKey::Shift, EKeys::Period));
+	UI_COMMAND(RightBoundsNextZeroCrossing, "Increase Right Loop Bounds To Next Zero Crossing", "Moves the end of the loop region earlier", EUserInterfaceActionType::Button, FInputChord(EModifierKey::Shift, EKeys::RightBracket));
+	UI_COMMAND(RightBoundsPreviousZeroCrossing, "Decrease Right Loop Bounds To Previous Zero Crossing", "Moves the end of the loop region later", EUserInterfaceActionType::Button, FInputChord(EModifierKey::Shift, EKeys::LeftBracket));
+
 	UI_COMMAND(BoundsIncrementIncrease, "Increase Bound Increment", "Increases the increment to shift the loop bounds", EUserInterfaceActionType::Button, FInputChord(EKeys::Equals));
 	UI_COMMAND(BoundsIncrementDecrease, "Decrease Bound Increment", "Increases the increment to shift the loop bounds", EUserInterfaceActionType::Button, FInputChord(EKeys::Hyphen));
 	UI_COMMAND(SelectNextLoop, "Select Next Loop", "Selects the next loop marker", EUserInterfaceActionType::Button, FInputChord(EKeys::F));
diff --git a/Engine/Plugins/Editor/WaveformEditor/Source/WaveformEditor/Public/WaveformEditorCommands.h b/Engine/Plugins/Editor/WaveformEditor/Source/WaveformEditor/Public/WaveformEditorCommands.h
--- a/Engine/Plugins/Editor/WaveformEditor/Source/WaveformEditor/Public/WaveformEditorCommands.h
+++ b/Engine/Plugins/Editor/WaveformEditor/Source/WaveformEditor/Public/WaveformEditorCommands.h
@@ -49,6 +49,10 @@ public:
 	TSharedPtr<FUICommandInfo> LeftBoundsDecrease;
 	TSharedPtr<FUICommandInfo> RightBoundsIncrease;
 	TSharedPtr<FUICommandInfo> RightBoundsDecrease;
+	TSharedPtr<FUICommandInfo> LeftBoundsNextZeroCrossing;
+	TSharedPtr<FUICommandInfo> LeftBoundsPreviousZeroCrossing;
+	TSharedPtr<FUICommandInfo> RightBoundsNextZeroCrossing;
+	TSharedPtr<FUICommandInfo> RightBoundsPreviousZeroCrossing;
 	TSharedPtr<FUICommandInfo> BoundsIncrementIncrease;
 	TSharedPtr<FUICommandInfo> BoundsIncrementDecrease;
 	TSharedPtr<FUICommandInfo> SelectNextLoop;
diff --git a/Engine/Plugins/Editor/WaveformEditor/Source/WaveformEditorWidgets/Public/WaveformEditorSequenceDataProvider.h b/Engine/Plugins/Editor/WaveformEditor/Source/WaveformEditorWidgets/Public/WaveformEditorSequenceDataProvider.h
--- a/Engine/Plugins/Editor/WaveformEditor/Source/WaveformEditorWidgets/Public/WaveformEditorSequenceDataProvider.h
+++ b/Engine/Plugins/Editor/WaveformEditor/Source/WaveformEditorWidgets/Public/WaveformEditorSequenceDataProvider.h
@@ -71,6 +71,8 @@ public:
 
 	UE_API const TRange<double> GetTransformedWaveformBounds() const;
 
+	UE_API const TArray<float> GetTransformedPCMData() const { return TransformedPCMData; };
+
 	/** Called when a new layer chain of transformations UI is created */
 	FOnLayersChainGenerated OnLayersChainGenerated;
 
diff --git a/Engine/Plugins/Editor/WaveformEditor/Source/WaveformTransformations/Private/WaveformTransformationMarkers.cpp b/Engine/Plugins/Editor/WaveformEditor/Source/WaveformTransformations/Private/WaveformTransformationMarkers.cpp
--- a/Engine/Plugins/Editor/WaveformEditor/Source/WaveformTransformations/Private/WaveformTransformationMarkers.cpp
+++ b/Engine/Plugins/Editor/WaveformEditor/Source/WaveformTransformations/Private/WaveformTransformationMarkers.cpp
@@ -189,7 +189,7 @@ void UWaveformTransformationMarkers::CycleMarkerLoopRegion(ELoopModificationCont
 }
 PRAGMA_ENABLE_DEPRECATION_WARNINGS
 
-void UWaveformTransformationMarkers::ModifyMarkerLoopRegion(ELoopModificationControls Modification)
+void UWaveformTransformationMarkers::ModifyMarkerLoopRegion(ELoopModificationControls Modification, const TArray<float>& PCMArray)
 {
 	check(Markers != nullptr);
 
@@ -206,8 +206,6 @@ void UWaveformTransformationMarkers::ModifyMarkerLoopRegion(ELoopModificationCon
 	check(SizeIncrements.Num() > 0);
 
 	int64 FramesToShift = static_cast<int64>((SizeIncrements[SelectedIncrement] / 1000.0f) * SampleRate);
-	int64 MaxFrames = static_cast<int64>(NumAvailableSamples) / static_cast<int64>(NumChannels);
-	check(MaxFrames > 0);
 
 	if (GEditor && GEditor->Trans)
 	{
@@ -257,12 +255,51 @@ void UWaveformTransformationMarkers::ModifyMarkerLoopRegion(ELoopModificationCon
 			--SelectedIncrement;
 		}
 	}
+	case ELoopModificationControls::LeftHandleNextZeroCrossing:
+	{
+		const int64 StartFrame = (LoopRegion->FramePosition) * NumChannels;
+		const int64 MaxFrames = 0;
+
+		const int64 FramesToMove = GetFramesToNextZeroCrossing(StartFrame, MaxFrames, PCMArray, false);
+
+		LoopRegion->FramePosition -= FramesToMove;
+		LoopRegion->FrameLength += FramesToMove;
+	}
+	break;
+	case ELoopModificationControls::LeftHandlePreviousZeroCrossing:
+	{
+		const int64 StartFrame = (LoopRegion->FramePosition) * NumChannels;
+		const int64 MaxFrames = (LoopRegion->FramePosition + LoopRegion->FrameLength) * NumChannels;
+
+		const int64 FramesToMove = GetFramesToNextZeroCrossing(StartFrame, MaxFrames, PCMArray, true);
+
+		LoopRegion->FramePosition += FramesToMove;
+		LoopRegion->FrameLength -= FramesToMove;
+	}
+	break;
+	case ELoopModificationControls::RightHandleNextZeroCrossing:
+	{
+		const int64 StartFrame = (LoopRegion->FramePosition + LoopRegion->FrameLength) * NumChannels;
+		const int64 MaxFrames = NumAvailableSamples;
+
+		LoopRegion->FrameLength += GetFramesToNextZeroCrossing(StartFrame, MaxFrames, PCMArray, true);
+	}
+	break;
+	case ELoopModificationControls::RightHandlePreviousZeroCrossing:
+	{
+		const int64 StartFrame = (LoopRegion->FramePosition + LoopRegion->FrameLength) * NumChannels;
+		const int64 MaxFrames = LoopRegion->FramePosition;
+
+		LoopRegion->FrameLength -= GetFramesToNextZeroCrossing(StartFrame, MaxFrames, PCMArray, false);
+	}
 	default:
 		break;
 	}
 
+	int64 MaxTotalFrames = static_cast<int64>(NumAvailableSamples) / static_cast<int64>(NumChannels);
+	check(MaxTotalFrames > 0);
 	// Prevent the frame position from exceeding the wave length
-	LoopRegion->FramePosition = FMath::Clamp(LoopRegion->FramePosition, 0, MaxFrames - 1);
+	LoopRegion->FramePosition = FMath::Clamp(LoopRegion->FramePosition, 0, MaxTotalFrames - 1);
 	// FrameLength can exceed MaxFrames because the loop will be cut by ProcessAudio
 	// but length must be >= MinLoopSize or it is no longer a loop
 	LoopRegion->FrameLength = FMath::Max(LoopRegion->FrameLength, Markers->MinLoopSize);
@@ -341,6 +378,82 @@ void UWaveformTransformationMarkers::CycleMarkerLoopRegion(ELoopModificationCont
 	OnTransformationChanged.ExecuteIfBound();
 }
 
+int64 UWaveformTransformationMarkers::GetFramesToNextZeroCrossing(int64 StartFrame, int64 EndFrame, const TArray<float>& PCMArray, bool bCheckIncrementing) const
+{
+	check(NumChannels > 0);
+	int64 FramesMoved = 0;
+
+	bool bAllFramesPositiveLastFrame = false;
+	bool bAllFramesNegativeLastFrame = false;
+
+	if (bCheckIncrementing)
+	{
+		for (int32 i = StartFrame; i < EndFrame; i += NumChannels)
+		{
+			if (CheckZeroCrossingFrame(i, PCMArray, bAllFramesPositiveLastFrame, bAllFramesNegativeLastFrame))
+			{
+				return FramesMoved;
+			}
+
+			++FramesMoved;
+		}
+	}
+	else
+	{
+		for (int32 i = StartFrame; i >= EndFrame; i -= NumChannels)
+		{
+			if (CheckZeroCrossingFrame(i, PCMArray, bAllFramesPositiveLastFrame, bAllFramesNegativeLastFrame))
+			{
+				return FramesMoved;
+			}
+
+			++FramesMoved;
+		}
+	}
+
+	return 0;
+}
+
+int64 UWaveformTransformationMarkers::CheckZeroCrossingFrame(int64 CurrentFrame, const TArray<float>& PCMArray, bool& bAllFramesPositiveLastFrame, bool& bAllFramesNegativeLastFrame) const
+{
+	check(NumChannels > 0);
+	check(CurrentFrame < PCMArray.Num());
+
+	int32 PositiveSampleCount = 0;
+	int32 NegativeSampleCount = 0;
+
+	// check if all channels have a zero crossing
+	for (int32 ChannelNum = 0; ChannelNum < NumChannels; ChannelNum++)
+	{
+		const float FrameValue = PCMArray[CurrentFrame + ChannelNum];
+
+		// Treat 0 values as being valid for zero crossings
+		if (FrameValue == 0)
+		{
+			++NegativeSampleCount;
+			++PositiveSampleCount;
+		}
+		else if (FrameValue < 0)
+		{
+			++NegativeSampleCount;
+		}
+		else
+		{
+			++PositiveSampleCount;
+		}
+	}
+
+	bool bAllSamplesPositive = PositiveSampleCount == NumChannels;
+	bool bAllSamplesNegative = NegativeSampleCount == NumChannels;
+
+	// Check if we have crossed the zero crossing on all channels between last frame and this frame
+	bool bZeroCrossingOccurred = (bAllSamplesPositive && bAllFramesNegativeLastFrame) || (bAllSamplesNegative && bAllFramesPositiveLastFrame);
+	bAllFramesPositiveLastFrame = bAllSamplesPositive;
+	bAllFramesNegativeLastFrame = bAllSamplesNegative;
+
+	return bZeroCrossingOccurred;
+}
+
 bool UWaveformTransformationMarkers::IsLoopRegionActive() const
 {
 	if (bIsPreviewingLoopRegion || EndLoopTime != -1.0)
diff --git a/Engine/Plugins/Editor/WaveformEditor/Source/WaveformTransformations/Public/WaveformTransformationMarkers.h b/Engine/Plugins/Editor/WaveformEditor/Source/WaveformTransformations/Public/WaveformTransformationMarkers.h
--- a/Engine/Plugins/Editor/WaveformEditor/Source/WaveformTransformations/Public/WaveformTransformationMarkers.h
+++ b/Engine/Plugins/Editor/WaveformEditor/Source/WaveformTransformations/Public/WaveformTransformationMarkers.h
@@ -22,6 +22,10 @@ enum class ELoopModificationControls : uint8
 	DecreaseIncrement,
 	SelectNextLoop,
 	SelectPreviousLoop,
+	LeftHandleNextZeroCrossing,
+	LeftHandlePreviousZeroCrossing,
+	RightHandleNextZeroCrossing,
+	RightHandlePreviousZeroCrossing,
 };
 
 //Used to make CuePoints work with PropertyHandles
@@ -118,10 +122,13 @@ public:
 	UE_API void CycleMarkerLoopRegion(ELoopModificationControls Modification) const;
 
 	// Shift marker and loop region bounds with toolkit commands bound in the waveform editor
-	UE_API void ModifyMarkerLoopRegion(ELoopModificationControls Modification);
+	UE_API void ModifyMarkerLoopRegion(ELoopModificationControls Modification, const TArray<float> &PCMArray);
 	// Set SelectedCue by cycling through cues and loops with toolkit commands bound in the waveform editor
 	UE_API void CycleMarkerLoopRegion(ELoopModificationControls Modification);
 
+	UE_API int64 GetFramesToNextZeroCrossing(int64 StartFrame, int64 EndFrame, const TArray<float>& PCMArray, bool bCheckIncrementing) const;
+	UE_API int64 CheckZeroCrossingFrame(int64 CurrentFrame, const TArray<float>& PCMArray, bool& bAllFramesPositiveLastFrame, bool& bAllFramesNegativeLastFrame) const;
+
 	UE_API bool IsLoopRegionActive() const;
 	UE_API void SetLoopPreviewing();
 
@@ -143,7 +150,7 @@ private:
 	UE_API void SetIsPreviewingLoopRegion(double InStartTime, double InEndTime, bool bIsPreviewing);
 	UE_API void SelectLoopRegionByKeyboard(const FKey& PressedKey);
 	
-	int32 SelectedIncrement = 0;
+	int32 SelectedIncrement = 2;
 
 	float SampleRate = 0.f;
 	float AvailableWaveformDuration = -1.f;
