Waveform Editor: zoom controls

diff --git a/Engine/Plugins/Editor/WaveformEditor/Source/WaveformEditor/Private/TransformedWaveformViewFactory.cpp b/Engine/Plugins/Editor/WaveformEditor/Source/WaveformEditor/Private/TransformedWaveformViewFactory.cpp
--- a/Engine/Plugins/Editor/WaveformEditor/Source/WaveformEditor/Private/TransformedWaveformViewFactory.cpp
+++ b/Engine/Plugins/Editor/WaveformEditor/Source/WaveformEditor/Private/TransformedWaveformViewFactory.cpp
@@ -38,9 +38,9 @@ FTransformedWaveformView FTransformedWaveformViewFactory::GetTransformedView(TOb
 		return TransportCoordinator->ConvertAbsoluteRatioToZoomed(InAnchorRatio);
 	};
 
-	TSharedPtr<SWaveformTransformationsOverlay> TransformationsOverlay = SNew(SWaveformTransformationsOverlay, WaveformDataProvider->GetTransformLayers()).AnchorsRatioConverter(TransformationsOverlayZoomConverter);
+	TSharedPtr<SWaveformTransformationsOverlay> TransformationsOverlay = SNew(SWaveformTransformationsOverlay, WaveformDataProvider->GetTransformLayers(), InZoomController).AnchorsRatioConverter(TransformationsOverlayZoomConverter);
 
-	TSharedPtr<STransformedWaveformViewPanel> WaveformPanel = SNew(STransformedWaveformViewPanel, SequenceView).TransformationsOverlay(TransformationsOverlay);
+	TSharedPtr<STransformedWaveformViewPanel> WaveformPanel = SNew(STransformedWaveformViewPanel, SequenceView).TransformationsOverlay(TransformationsOverlay).ZoomController(InZoomController);
 	SetUpWaveformPanelInteractions(WaveformPanel.ToSharedRef(), TransportCoordinator, InZoomController, TransformationsOverlay);
 
 	WaveformDataProvider->OnLayersChainGenerated.AddSP(TransformationsOverlay.Get(), &SWaveformTransformationsOverlay::OnLayerChainGenerated);
diff --git a/Engine/Plugins/Editor/WaveformEditor/Source/WaveformEditor/Private/WaveformEditor.cpp b/Engine/Plugins/Editor/WaveformEditor/Source/WaveformEditor/Private/WaveformEditor.cpp
--- a/Engine/Plugins/Editor/WaveformEditor/Source/WaveformEditor/Private/WaveformEditor.cpp
+++ b/Engine/Plugins/Editor/WaveformEditor/Source/WaveformEditor/Private/WaveformEditor.cpp
@@ -839,7 +839,7 @@ bool FWaveformEditor::CreateTransportController()
 
 bool FWaveformEditor::InitializeZoom()
 {
-	ZoomManager = MakeShared<FWaveformEditorZoomController>();
+	ZoomManager = MakeShared<FWaveformEditorZoomController>(TransportCoordinator);
 
 	check(TransportCoordinator)
 	ZoomManager->OnZoomRatioChanged.AddSP(TransportCoordinator.ToSharedRef(), &FSparseSampledSequenceTransportCoordinator::SetZoomRatio);
diff --git a/Engine/Plugins/Editor/WaveformEditor/Source/WaveformEditorWidgets/Private/STransformedWaveformViewPanel.cpp b/Engine/Plugins/Editor/WaveformEditor/Source/WaveformEditorWidgets/Private/STransformedWaveformViewPanel.cpp
--- a/Engine/Plugins/Editor/WaveformEditor/Source/WaveformEditorWidgets/Private/STransformedWaveformViewPanel.cpp
+++ b/Engine/Plugins/Editor/WaveformEditor/Source/WaveformEditorWidgets/Private/STransformedWaveformViewPanel.cpp
@@ -10,6 +10,7 @@
 #include "SSampledSequenceValueGridOverlay.h"
 #include "SWaveformEditorInputRoutingOverlay.h"
 #include "SWaveformTransformationsOverlay.h"
+#include "WaveformEditorZoomController.h"
 #include "WaveformEditorGridData.h"
 #include "WaveformEditorStyle.h"
 #include "WaveformEditorWidgetsSettings.h"
@@ -38,6 +39,11 @@ void STransformedWaveformViewPanel::Construct(const FArguments& InArgs, const FF
 		WaveformTransformationsOverlay = InArgs._TransformationsOverlay;
 	}
 
+	if (InArgs._ZoomController)
+	{
+		ZoomController = InArgs._ZoomController;
+	}
+
 	SetupPlayheadOverlay();
 	SetUpInputRoutingOverlay();
 	SetUpTimeRuler(GridData.ToSharedRef());
@@ -73,8 +79,6 @@ void STransformedWaveformViewPanel::CreateLayout()
 		WaveformViewer.ToSharedRef()
 	];
 
-
-
 	if (WaveformTransformationsOverlay)
 	{
 		WaveformView->AddSlot()
@@ -251,6 +255,10 @@ void STransformedWaveformViewPanel::Tick(const FGeometry& AllottedGeometry, cons
 
 	UpdatePlayheadPosition(PaintedWidth);
 
+	if (ZoomController)
+	{
+		ZoomController->CheckBounds(AllottedGeometry);
+	}
 }
 
 void STransformedWaveformViewPanel::UpdatePlayheadPosition(const float PaintedWidth)
diff --git a/Engine/Plugins/Editor/WaveformEditor/Source/WaveformEditorWidgets/Private/SWaveformTransformationsOverlay.cpp b/Engine/Plugins/Editor/WaveformEditor/Source/WaveformEditorWidgets/Private/SWaveformTransformationsOverlay.cpp
--- a/Engine/Plugins/Editor/WaveformEditor/Source/WaveformEditorWidgets/Private/SWaveformTransformationsOverlay.cpp
+++ b/Engine/Plugins/Editor/WaveformEditor/Source/WaveformEditorWidgets/Private/SWaveformTransformationsOverlay.cpp
@@ -5,12 +5,16 @@
 #include "AudioWidgetsUtils.h"
 #include "SWaveformTransformationRenderLayer.h"
 #include "SparseSampledSequenceTransportCoordinator.h"
+#include "WaveformEditorZoomController.h"
 #include "Widgets/SOverlay.h"
 
-void SWaveformTransformationsOverlay::Construct(const FArguments& InArgs, TArrayView<const FTransformationLayerRenderInfo> InTransformationRenderers)
+void SWaveformTransformationsOverlay::Construct(const FArguments& InArgs, TArrayView<const FTransformationLayerRenderInfo> InTransformationRenderers, TSharedPtr<FWaveformEditorZoomController> InZoomController)
 {
 	TransformationRenderers = InTransformationRenderers;
 	AnchorsRatioConverter = InArgs._AnchorsRatioConverter;
+
+	check(InZoomController.IsValid());
+	ZoomController = InZoomController;
 	CreateLayout();
 }
 
@@ -93,16 +97,19 @@ void SWaveformTransformationsOverlay::UpdateLayerConstraints()
 
 FReply SWaveformTransformationsOverlay::OnMouseButtonDown(const FGeometry& MyGeometry, const FPointerEvent& MouseEvent)
 {
+	ZoomController->OnMouseButtonDown(MyGeometry, MouseEvent);
 	return AudioWidgetsUtils::RouteMouseInput(&SWidget::OnMouseButtonDown, MouseEvent, MakeArrayView(TransformationLayers.GetData(), TransformationLayers.Num()), false);
 }
 
 FReply SWaveformTransformationsOverlay::OnMouseButtonUp(const FGeometry& MyGeometry, const FPointerEvent& MouseEvent)
 {
+	ZoomController->OnMouseButtonUp(MyGeometry, MouseEvent);
 	return AudioWidgetsUtils::RouteMouseInput(&SWidget::OnMouseButtonUp, MouseEvent, MakeArrayView(TransformationLayers.GetData(), TransformationLayers.Num()), false);
 }
 
 FReply SWaveformTransformationsOverlay::OnMouseMove(const FGeometry& MyGeometry, const FPointerEvent& MouseEvent)
 {
+	ZoomController->OnMouseMove(MyGeometry, MouseEvent);
 	return AudioWidgetsUtils::RouteMouseInput(&SWidget::OnMouseMove, MouseEvent, MakeArrayView(TransformationLayers.GetData(), TransformationLayers.Num()), false);
 }
 
diff --git a/Engine/Plugins/Editor/WaveformEditor/Source/WaveformEditorWidgets/Private/WaveformEditorZoomController.cpp b/Engine/Plugins/Editor/WaveformEditor/Source/WaveformEditorWidgets/Private/WaveformEditorZoomController.cpp
--- a/Engine/Plugins/Editor/WaveformEditor/Source/WaveformEditorWidgets/Private/WaveformEditorZoomController.cpp
+++ b/Engine/Plugins/Editor/WaveformEditor/Source/WaveformEditorWidgets/Private/WaveformEditorZoomController.cpp
@@ -2,6 +2,16 @@
 
 #include "WaveformEditorZoomController.h"
 
+#include "SparseSampledSequenceTransportCoordinator.h"
+#include "WaveformEditorStyle.h"
+#include <Blueprint/WidgetLayoutLibrary.h>
+
+FWaveformEditorZoomController::FWaveformEditorZoomController(TSharedPtr<FSparseSampledSequenceTransportCoordinator> TransportCoordinator)
+{
+	check(TransportCoordinator);
+	TransportController = TransportCoordinator;
+}
+
 void FWaveformEditorZoomController::ZoomIn()
 {
 	if (CanZoomIn())
@@ -49,6 +59,62 @@ float FWaveformEditorZoomController::GetZoomRatio() const
 	return 1 - ConvertZoomLevelToLogRatio();
 }
 
+void FWaveformEditorZoomController::CheckBounds(const FGeometry& AllottedGeometry)
+{
+	if (bIsPanning)
+	{
+		FVector2D MouseAbsolutePosition = UWidgetLayoutLibrary::GetMousePositionOnPlatform();
+
+		if (!AllottedGeometry.IsUnderLocation(MouseAbsolutePosition))
+		{
+			bIsPanning = false;
+		}
+	}
+}
+
+FReply FWaveformEditorZoomController::OnMouseButtonDown(const FGeometry& MyGeometry, const FPointerEvent& MouseEvent)
+{
+	if (MouseEvent.GetEffectingButton() == EKeys::RightMouseButton)
+	{
+		const FVector2D LocalCursorPosition = MyGeometry.AbsoluteToLocal(MouseEvent.GetScreenSpacePosition());
+		LocalCursorXPosition = FMath::Clamp(LocalCursorPosition.X, 0.f, MyGeometry.GetLocalSize().X);
+
+		bIsPanning = true;
+	}
+
+	return FReply::Unhandled();
+}
+
+FReply FWaveformEditorZoomController::OnMouseButtonUp(const FGeometry& MyGeometry, const FPointerEvent& MouseEvent)
+{
+	if (MouseEvent.GetEffectingButton() == EKeys::RightMouseButton)
+	{
+		bIsPanning = false;
+	}
+
+	return FReply::Unhandled();
+}
+
+FReply FWaveformEditorZoomController::OnMouseMove(const FGeometry& MyGeometry, const FPointerEvent& MouseEvent)
+{
+	if (bIsPanning)
+	{
+		check(TransportController);
+		check(MyGeometry.GetLocalSize().X > 0);
+
+		const FVector2D LocalCursorPosition = MyGeometry.AbsoluteToLocal(MouseEvent.GetScreenSpacePosition());
+		const float LocalCursorXClamped = FMath::Clamp(LocalCursorPosition.X, 0.f, MyGeometry.GetLocalSize().X);
+		const float MovementDifference = (LocalCursorXClamped - LocalCursorXPosition);
+		const float MoveRatio = MovementDifference / MyGeometry.GetLocalSize().X;
+
+		TransportController->MoveWaveformView(-MoveRatio);
+
+		LocalCursorXPosition = LocalCursorXClamped;
+	}
+
+	return FReply::Unhandled();
+}
+
 void FWaveformEditorZoomController::ApplyZoom()
 {
 	OnZoomRatioChanged.Broadcast(1 - ConvertZoomLevelToLogRatio());
diff --git a/Engine/Plugins/Editor/WaveformEditor/Source/WaveformEditorWidgets/Public/STransformedWaveformViewPanel.h b/Engine/Plugins/Editor/WaveformEditor/Source/WaveformEditorWidgets/Public/STransformedWaveformViewPanel.h
--- a/Engine/Plugins/Editor/WaveformEditor/Source/WaveformEditorWidgets/Public/STransformedWaveformViewPanel.h
+++ b/Engine/Plugins/Editor/WaveformEditor/Source/WaveformEditorWidgets/Public/STransformedWaveformViewPanel.h
@@ -35,6 +35,8 @@ public:
 		
 		SLATE_ARGUMENT(TSharedPtr<SWaveformTransformationsOverlay>, TransformationsOverlay)
 
+		SLATE_ARGUMENT(TSharedPtr<FWaveformEditorZoomController>, ZoomController)
+
 		SLATE_ARGUMENT(FPointerEventHandler, OnPlayheadOverlayMouseButtonUp)
 
 		SLATE_ARGUMENT(FPointerEventHandler, OnTimeRulerMouseButtonUp)
@@ -88,6 +90,7 @@ private:
 	TSharedPtr<SPlayheadOverlay> PlayheadOverlay;
 	TSharedPtr<SSampledSequenceValueGridOverlay> ValueGridOverlay;
 	TSharedPtr<SBorder> BackgroundBorder;
+	TSharedPtr<FWaveformEditorZoomController> ZoomController;
 
 	float CachedPixelWidth = 0.f;
 
diff --git a/Engine/Plugins/Editor/WaveformEditor/Source/WaveformEditorWidgets/Public/SWaveformTransformationsOverlay.h b/Engine/Plugins/Editor/WaveformEditor/Source/WaveformEditorWidgets/Public/SWaveformTransformationsOverlay.h
--- a/Engine/Plugins/Editor/WaveformEditor/Source/WaveformEditorWidgets/Public/SWaveformTransformationsOverlay.h
+++ b/Engine/Plugins/Editor/WaveformEditor/Source/WaveformEditorWidgets/Public/SWaveformTransformationsOverlay.h
@@ -2,6 +2,7 @@
 
 #pragma once
 
+#include "WaveformEditorZoomController.h"
 #include "Widgets/SCompoundWidget.h"
 #include "Widgets/Layout/SConstraintCanvas.h"
 
@@ -29,7 +30,7 @@ public:
 
 	SLATE_END_ARGS()
 
-	UE_API void Construct(const FArguments& InArgs, TArrayView< const FTransformationLayerRenderInfo> InTransformationRenderers);
+	UE_API void Construct(const FArguments& InArgs, TArrayView< const FTransformationLayerRenderInfo> InTransformationRenderers, TSharedPtr<FWaveformEditorZoomController> InZoomController);
 	UE_API void OnLayerChainGenerated(FTransformationLayerRenderInfo* FirstLayerPtr, const int32 NLayers);
 	UE_API void UpdateLayerConstraints();
 
@@ -48,6 +49,7 @@ private:
 	TArray<SConstraintCanvas::FSlot*> LayersSlots;
 	TArrayView<const FTransformationLayerRenderInfo> TransformationRenderers;
 	TFunction<float(const float)> AnchorsRatioConverter;
+	TSharedPtr<FWaveformEditorZoomController> ZoomController;
 };
 
 #undef UE_API
diff --git a/Engine/Plugins/Editor/WaveformEditor/Source/WaveformEditorWidgets/Public/WaveformEditorZoomController.h b/Engine/Plugins/Editor/WaveformEditor/Source/WaveformEditorWidgets/Public/WaveformEditorZoomController.h
--- a/Engine/Plugins/Editor/WaveformEditor/Source/WaveformEditorWidgets/Public/WaveformEditorZoomController.h
+++ b/Engine/Plugins/Editor/WaveformEditor/Source/WaveformEditorWidgets/Public/WaveformEditorZoomController.h
@@ -2,7 +2,12 @@
 
 #pragma once
 
+#include "AudioWidgetsSlateTypes.h"
 #include "Delegates/Delegate.h"
+#include "SparseSampledSequenceTransportCoordinator.h"
+#include "Styling/ISlateStyle.h"
+#include "Styling/SlateWidgetStyleAsset.h"
+#include "Widgets/SLeafWidget.h"
 
 #define UE_API WAVEFORMEDITORWIDGETS_API
 
@@ -11,7 +16,12 @@ DECLARE_MULTICAST_DELEGATE_OneParam(FOnZoomRatioChanged, const float /* New Zoom
 class FWaveformEditorZoomController
 {
 public:
-	FWaveformEditorZoomController() = default;
+	SLATE_BEGIN_ARGS(FWaveformEditorZoomController)
+		{
+		}
+
+	SLATE_END_ARGS()
+	UE_API FWaveformEditorZoomController(TSharedPtr<FSparseSampledSequenceTransportCoordinator> TransportCoordinator);
 
 	UE_API void ZoomIn();
 	UE_API bool CanZoomIn() const;
@@ -21,16 +31,26 @@ public:
 
 	UE_API float GetZoomRatio() const;
 
+	UE_API void CheckBounds(const FGeometry& AllottedGeometry);
+	UE_API FReply OnMouseButtonDown(const FGeometry& MyGeometry, const FPointerEvent& MouseEvent);
+	UE_API FReply OnMouseButtonUp(const FGeometry& MyGeometry, const FPointerEvent& MouseEvent);
+	UE_API FReply OnMouseMove(const FGeometry& MyGeometry, const FPointerEvent& MouseEvent);
+
 	FOnZoomRatioChanged OnZoomRatioChanged;
 
 private:
 	UE_API void ApplyZoom();
 	UE_API float ConvertZoomLevelToLogRatio() const;
-	
+
 	float ZoomLevelInitValue = 1.f;
 	float ZoomLevel = ZoomLevelInitValue;
 	float ZoomLevelStep = 2.f;
 	const uint32 LogRatioBase = 100;
+
+	bool bIsPanning = false;
+	float LocalCursorXPosition = 0.0f;
+
+	TSharedPtr<FSparseSampledSequenceTransportCoordinator> TransportController = nullptr;
 };
 
 #undef UE_API
diff --git a/Engine/Plugins/Editor/WaveformEditor/Source/WaveformEditorWidgets/WaveformEditorWidgets.build.cs b/Engine/Plugins/Editor/WaveformEditor/Source/WaveformEditorWidgets/WaveformEditorWidgets.build.cs
--- a/Engine/Plugins/Editor/WaveformEditor/Source/WaveformEditorWidgets/WaveformEditorWidgets.build.cs
+++ b/Engine/Plugins/Editor/WaveformEditor/Source/WaveformEditorWidgets/WaveformEditorWidgets.build.cs
@@ -26,6 +26,7 @@ namespace UnrealBuildTool.Rules
 					"Slate",
 					"SlateCore", 
 					"ToolMenus",
+					"UMG",
 					"WaveformTransformationsWidgets",
 				}
 			);
diff --git a/Engine/Plugins/Runtime/AudioWidgets/Source/AudioWidgets/Private/SparseSampledSequenceTransportCoordinator.cpp b/Engine/Plugins/Runtime/AudioWidgets/Source/AudioWidgets/Private/SparseSampledSequenceTransportCoordinator.cpp
--- a/Engine/Plugins/Runtime/AudioWidgets/Source/AudioWidgets/Private/SparseSampledSequenceTransportCoordinator.cpp
+++ b/Engine/Plugins/Runtime/AudioWidgets/Source/AudioWidgets/Private/SparseSampledSequenceTransportCoordinator.cpp
@@ -101,6 +101,30 @@ void FSparseSampledSequenceTransportCoordinator::Stop()
 	SetProgressRatio(0.f);
 }
 
+void FSparseSampledSequenceTransportCoordinator::MoveWaveformView(float MoveAmount)
+{
+	double MinDisplayRangeValue = DisplayRange.GetLowerBoundValue() + (MoveAmount * ZoomRatio);
+	double MaxDisplayRangeValue = DisplayRange.GetUpperBoundValue() + (MoveAmount * ZoomRatio);
+
+	//Clamp zoom bounds while maintaining same width if bounds go outside 0-1 range
+	if (MinDisplayRangeValue < 0.0)
+	{
+		const double ExceedingDelta = FMath::Abs(MinDisplayRangeValue);
+		MaxDisplayRangeValue = FMath::Clamp(MaxDisplayRangeValue + ExceedingDelta, 0.0, 1.0);
+		MinDisplayRangeValue = 0.0;
+	}
+
+	if (MaxDisplayRangeValue > 1.0)
+	{
+		const double ExceedingDelta = MaxDisplayRangeValue - 1.0;
+
+		MinDisplayRangeValue = FMath::Clamp(MinDisplayRangeValue - ExceedingDelta, 0.0, 1.0);
+		MaxDisplayRangeValue = 1.0;
+	}
+
+	UpdateDisplayRange(MinDisplayRangeValue, MaxDisplayRangeValue);
+}
+
 const float FSparseSampledSequenceTransportCoordinator::GetFocusPoint() const
 {
 	return FocusPoint;
diff --git a/Engine/Plugins/Runtime/AudioWidgets/Source/AudioWidgets/Public/SparseSampledSequenceTransportCoordinator.h b/Engine/Plugins/Runtime/AudioWidgets/Source/AudioWidgets/Public/SparseSampledSequenceTransportCoordinator.h
--- a/Engine/Plugins/Runtime/AudioWidgets/Source/AudioWidgets/Public/SparseSampledSequenceTransportCoordinator.h
+++ b/Engine/Plugins/Runtime/AudioWidgets/Source/AudioWidgets/Public/SparseSampledSequenceTransportCoordinator.h
@@ -38,6 +38,8 @@ public:
 	UE_API void UpdatePlaybackRange(const TRange<double>& NewRange);
 	UE_API void Stop();
 
+	UE_API void MoveWaveformView(float MoveAmount);
+
 private:
 	inline void MoveFocusPoint(const double InFocusPoint);
 	UE_API void UpdateZoomRatioAndDisplayRange(const float NewZoomRatio);
