[Waveform Editor] - Add inverted S fade function
[Waveform Editor] - Crash on reloading an asset with a deleted Trim Fade

diff --git a/Engine/Plugins/Editor/WaveformEditor/Source/WaveformEditor/Private/WaveformEditor.cpp b/Engine/Plugins/Editor/WaveformEditor/Source/WaveformEditor/Private/WaveformEditor.cpp
--- a/Engine/Plugins/Editor/WaveformEditor/Source/WaveformEditor/Private/WaveformEditor.cpp
+++ b/Engine/Plugins/Editor/WaveformEditor/Source/WaveformEditor/Private/WaveformEditor.cpp
@@ -247,27 +247,41 @@ FWaveformEditor::~FWaveformEditor()
 		{
 			TrimFadeTransformation->OnTransformationRenderChanged.Unbind();
 		}
-	}
 
-	if (FReimportManager::Instance())
-	{
-		FReimportManager::Instance()->OnPostReimport().RemoveAll(this);
-	}
+		// If the user has not used the trimfade then remove the transformation to avoid unnecessary transformations in the chain
+		check(TrimFadeTransformation != nullptr);
+		check(SoundWave != nullptr);
 
-	// If the user has not used the trimfade then remove the transformation to avoid unecessary transformations in the chain
-	check(TrimFadeTransformation != nullptr);
-	check(SoundWave != nullptr);
+		// Account for a user trimming a single frame by having a 1/2 frame error tolerance
+		const uint32 ImportedSampleRate = SoundWave->GetImportedSampleRate();
+		check(ImportedSampleRate > 0);
 
-	// Account for a user trimming a single frame by having a 1/2 frame error tolerance
-	uint32 ImportedSampleRate = SoundWave->GetImportedSampleRate();
-	check(ImportedSampleRate > 0);
+		const double ErrorTolerance = (1.0f / ImportedSampleRate) * 0.5f;
 
-	const double ErrorTolerance = (1.0f / ImportedSampleRate) * 0.5f;
+		if (FMath::IsNearlyEqual(TrimFadeTransformation->StartTime, 0, ErrorTolerance) && FMath::IsNearlyEqual(TrimFadeTransformation->EndTime, SoundWave->Duration, ErrorTolerance))
+		{
+			bool bRemoveTrimFade = true;
 
-	if (FMath::IsNearlyEqual(TrimFadeTransformation->StartTime, 0, ErrorTolerance) && FMath::IsNearlyEqual(TrimFadeTransformation->EndTime, SoundWave->Duration, ErrorTolerance)
-		&& FMath::IsNearlyEqual(TrimFadeTransformation->StartFadeTime, 0, ErrorTolerance) && FMath::IsNearlyEqual(TrimFadeTransformation->EndFadeTime, 0, ErrorTolerance))
+			if ((TrimFadeTransformation->FadeFunctions.FadeIn && !FMath::IsNearlyEqual(TrimFadeTransformation->FadeFunctions.FadeIn->Duration, 0, ErrorTolerance))
+				|| (TrimFadeTransformation->FadeFunctions.FadeOut && !FMath::IsNearlyEqual(TrimFadeTransformation->FadeFunctions.FadeOut->Duration, 0, ErrorTolerance)))
+			{
+				// If any FadeFunctions are present with a non-zero duration, do not remove the transformation
+				bRemoveTrimFade = false;
+			}
+			
+			PRAGMA_DISABLE_DEPRECATION_WARNINGS
+			// If no FadeFunctions are present, check the deprecated fade properties before removing the trimfade
+			if (bRemoveTrimFade && FMath::IsNearlyEqual(TrimFadeTransformation->StartFadeTime, 0, ErrorTolerance) && FMath::IsNearlyEqual(TrimFadeTransformation->EndFadeTime, 0, ErrorTolerance))
+			{
+				SoundWave->Transformations.Remove(TrimFadeTransformation);
+			}
+			PRAGMA_ENABLE_DEPRECATION_WARNINGS
+		}
+	}
+
+	if (FReimportManager::Instance())
 	{
-		SoundWave->Transformations.Remove(TrimFadeTransformation);
+		FReimportManager::Instance()->OnPostReimport().RemoveAll(this);
 	}
 }
 
@@ -349,6 +363,7 @@ TObjectPtr<UWaveformTransformationTrimFade> FWaveformEditor::GetOrAddTrimFadeTra
 
 		TrimFadeTransformation = NewObject<UWaveformTransformationTrimFade>(SoundWave, UWaveformTransformationTrimFade::StaticClass(), NAME_None, MaskedOuterFlags);
 
+		check(TrimFadeTransformation);
 		if (bMarkModify)
 		{
 			SoundWave->Modify();
@@ -359,7 +374,17 @@ TObjectPtr<UWaveformTransformationTrimFade> FWaveformEditor::GetOrAddTrimFadeTra
 		{
 			SoundWave->Transformations.Add(TrimFadeTransformation);
 		}
+
+		// Add default fade curves
+		TrimFadeTransformation->FadeFunctions.FadeIn = NewObject<UFadeFunction>(TrimFadeTransformation.Get(), FadeInType.Get());
+		check(TrimFadeTransformation->FadeFunctions.FadeIn);
+		TrimFadeTransformation->FadeFunctions.FadeOut = NewObject<UFadeFunction>(TrimFadeTransformation.Get(), FadeOutType.Get());
+		check(TrimFadeTransformation->FadeFunctions.FadeOut);
+
+		TrimFadeTransformation->FadeFunctions.FadeIn->Duration = 0.f;
+		TrimFadeTransformation->FadeFunctions.FadeOut->Duration = 0.f;
 	}
+
 	return TrimFadeTransformation;
 }
 
@@ -452,26 +477,64 @@ void FWaveformEditor::ToggleFadeIn()
 
 	TrimFadeTransformation->Modify();
 	
-	//if fade duration is 0, assume there is no fade present and toggle on the fade with the curve data selected in the menu
-	if (TrimFadeTransformation->StartFadeTime == 0.0f)
+	if (TrimFadeTransformation->FadeFunctions.FadeIn)
 	{
-		check(DefaultFadeInAmount >= 0.0f);
-		check(FadeInCurve >= -0.1f && FadeInCurve <= 10);
-
-		if (CachedFadeInAmount <= 0.0f)
+		if (TrimFadeTransformation->FadeFunctions.FadeIn->Duration == 0.0f)
 		{
-			CachedFadeInAmount = DefaultFadeInAmount;
+			check(DefaultFadeInAmount >= 0.0f);
+
+			PRAGMA_DISABLE_DEPRECATION_WARNINGS
+			check(FadeInCurve >= UWaveformTransformationTrimFade::MinFadeCurve && FadeInCurve <= UWaveformTransformationTrimFade::MaxFadeCurve);
+			PRAGMA_ENABLE_DEPRECATION_WARNINGS
+
+			if (CachedFadeInAmount <= 0.0f)
+			{
+				CachedFadeInAmount = DefaultFadeInAmount;
+			}
+
+			TrimFadeTransformation->FadeFunctions.FadeIn->Duration = CachedFadeInAmount;
 		}
+		else //if fade duration is not 0, assume there is a fade applied and toggle fade "off" by setting fade time to 0
+		{
+			CachedFadeInAmount = TrimFadeTransformation->FadeFunctions.FadeIn->Duration;
 
-		TrimFadeTransformation->StartFadeTime = CachedFadeInAmount;
-		TrimFadeTransformation->StartFadeCurve = FadeInCurve;
+			//No need to set FadeCurve when toggling off. This will preserve user curve data if they later manually change the StartFadeTime.
+			TrimFadeTransformation->FadeFunctions.FadeIn->Duration = 0.0f;
+		}
 	}
-	else //if fade duration is not 0, assume there is a fade applied and toggle fade "off" by setting fade time to 0
+	else // TrimFadeTransformation->FadeFunctions.FadeIn is nullptr
 	{
-		CachedFadeInAmount = TrimFadeTransformation->StartFadeTime;
+		PRAGMA_DISABLE_DEPRECATION_WARNINGS
+		//if fade duration is 0, assume there is no fade present and toggle on the fade with the curve data selected in the menu
+		if (TrimFadeTransformation->StartFadeTime == 0.0f)
+		{
+			// No cached amount means we can instantiate a new FadeIn Function object
+			if (CachedFadeInAmount == 0.f)
+			{
+				TrimFadeTransformation->FadeFunctions.FadeIn = NewObject<UFadeFunction>(TrimFadeTransformation.Get(), FadeInType.Get());
+			}
+			else
+			{
+				check(DefaultFadeInAmount >= 0.0f);
+				check(FadeInCurve >= UWaveformTransformationTrimFade::MinFadeCurve && FadeInCurve <= UWaveformTransformationTrimFade::MaxFadeCurve);
 
-		//No need to set FadeCurve when toggling off. This will preserve user curve data if they later manually change the StartFadeTime.
-		TrimFadeTransformation->StartFadeTime = 0.0f;
+				if (CachedFadeInAmount <= 0.0f)
+				{
+					CachedFadeInAmount = DefaultFadeInAmount;
+				}
+
+				TrimFadeTransformation->StartFadeTime = CachedFadeInAmount;
+				TrimFadeTransformation->StartFadeCurve = FadeInCurve;
+			}
+		}
+		else //if fade duration is not 0, assume there is a fade applied and toggle fade "off" by setting fade time to 0
+		{
+			CachedFadeInAmount = TrimFadeTransformation->StartFadeTime;
+
+			//No need to set FadeCurve when toggling off. This will preserve user curve data if they later manually change the StartFadeTime.
+			TrimFadeTransformation->StartFadeTime = 0.0f;
+		}
+		PRAGMA_ENABLE_DEPRECATION_WARNINGS
 	}
 
 	NotifyPostTransformationChange(EPropertyChangeType::ValueSet);
@@ -496,26 +559,64 @@ void FWaveformEditor::ToggleFadeOut()
 
 	TrimFadeTransformation->Modify();
 
-	//if fade duration is 0, assume there is no fade present and toggle on the fade with the curve data selected in the menu
-	if (TrimFadeTransformation->EndFadeTime == 0.0f)
+	if (TrimFadeTransformation->FadeFunctions.FadeOut)
 	{
-		check(DefaultFadeOutAmount >= 0.0f);
-		check(FadeOutCurve >= -0.1f && FadeOutCurve <= 10);
+		if (TrimFadeTransformation->FadeFunctions.FadeOut->Duration == 0.0f)
+		{
+			check(DefaultFadeOutAmount >= 0.0f);
 
-		if (CachedFadeOutAmount <= 0.0f)
+			PRAGMA_DISABLE_DEPRECATION_WARNINGS
+			check(FadeOutCurve >= UWaveformTransformationTrimFade::MinFadeCurve && FadeOutCurve <= UWaveformTransformationTrimFade::MaxFadeCurve);
+			PRAGMA_ENABLE_DEPRECATION_WARNINGS
+
+			if (CachedFadeOutAmount <= 0.0f)
+			{
+				CachedFadeOutAmount = DefaultFadeOutAmount;
+			}
+
+			TrimFadeTransformation->FadeFunctions.FadeOut->Duration = CachedFadeOutAmount;
+		}
+		else //if fade duration is not 0, assume there is a fade applied and toggle fade "off" by setting fade time to 0
 		{
-			CachedFadeOutAmount = DefaultFadeOutAmount;
+			CachedFadeOutAmount = TrimFadeTransformation->FadeFunctions.FadeOut->Duration;
+
+			//No need to set FadeCurve when toggling off. This will preserve user curve data if they later manually change the StartFadeTime.
+			TrimFadeTransformation->FadeFunctions.FadeOut->Duration = 0.0f;
 		}
-		
-		TrimFadeTransformation->EndFadeTime = CachedFadeOutAmount;
-		TrimFadeTransformation->EndFadeCurve = FadeOutCurve;
 	}
-	else //if fade duration is not 0, assume there is a fade applied and toggle fade "off" by setting fade time to 0
+	else // TrimFadeTransformation->FadeFunctions.FadeOut is nullptr
 	{
-		CachedFadeOutAmount = TrimFadeTransformation->EndFadeTime;
+		PRAGMA_DISABLE_DEPRECATION_WARNINGS
+		//if fade duration is 0, assume there is no fade present and toggle on the fade with the curve data selected in the menu
+		if (TrimFadeTransformation->EndFadeTime == 0.0f)
+		{
+			// No cached amount means we can instantiate a new FadeIn Function object
+			if (CachedFadeOutAmount == 0.f)
+			{
+				TrimFadeTransformation->FadeFunctions.FadeOut = NewObject<UFadeFunction>(TrimFadeTransformation.Get(), FadeOutType.Get());
+			}
+			else
+			{
+				check(DefaultFadeOutAmount >= 0.0f);
+				check(FadeOutCurve >= UWaveformTransformationTrimFade::MinFadeCurve && FadeOutCurve <= UWaveformTransformationTrimFade::MaxFadeCurve);
+
+				if (CachedFadeOutAmount <= 0.0f)
+				{
+					CachedFadeOutAmount = DefaultFadeOutAmount;
+				}
+
+				TrimFadeTransformation->EndFadeTime = CachedFadeOutAmount;
+				TrimFadeTransformation->EndFadeCurve = FadeOutCurve;
+			}
+		}
+		else //if fade duration is not 0, assume there is a fade applied and toggle fade "off" by setting fade time to 0
+		{
+			CachedFadeOutAmount = TrimFadeTransformation->EndFadeTime;
 
-		//No need to set FadeCurve when toggling off. This will preserve user curve data if they later manually change the EndFadeTime.
-		TrimFadeTransformation->EndFadeTime = 0.0f;
+			//No need to set FadeCurve when toggling off. This will preserve user curve data if they later manually change the EndFadeTime.
+			TrimFadeTransformation->EndFadeTime = 0.0f;
+		}
+		PRAGMA_ENABLE_DEPRECATION_WARNINGS
 	}
 
 	NotifyPostTransformationChange(EPropertyChangeType::ValueSet);
@@ -1374,7 +1475,7 @@ bool FWaveformEditor::BindCommands()
 		Commands.ReimportModeNewFile,
 		FExecuteAction::CreateLambda([this] { ReimportMode = EWaveEditorReimportMode::SelectFile; }),
 		FCanExecuteAction::CreateLambda([this] { return CanExecuteReimport(); }),
-		FIsActionChecked::CreateLambda([this] {return ReimportMode == EWaveEditorReimportMode::SelectFile; }));
+		FIsActionChecked::CreateLambda([this] { return ReimportMode == EWaveEditorReimportMode::SelectFile; }));
 
 	ToolkitCommands->MapAction(
 		Commands.ToggleFadeIn,
@@ -1383,27 +1484,43 @@ bool FWaveformEditor::BindCommands()
 
 	ToolkitCommands->MapAction(
 		Commands.FadeInLinear,
-		FExecuteAction::CreateLambda([this] { FadeInCurve = UWaveformTransformationTrimFade::FadeModeToCurveValueMap[EWaveEditorFadeMode::Linear]; }),
-		FCanExecuteAction::CreateLambda([this] {return true; }),
-		FIsActionChecked::CreateLambda([this] {return FadeInCurve == UWaveformTransformationTrimFade::FadeModeToCurveValueMap[EWaveEditorFadeMode::Linear]; }));
+		FExecuteAction::CreateLambda([this] 
+			{ 
+				FadeInCurve = UWaveformTransformationTrimFade::FadeModeToCurveValueMap[EWaveEditorFadeMode::Linear];
+				FadeInType = UWaveformTransformationTrimFade::FadeModeToFadeFunctionMap[EWaveEditorFadeMode::Linear];
+			}),
+		FCanExecuteAction::CreateLambda([this] { return true; }),
+		FIsActionChecked::CreateLambda([this] { return FadeInType == UWaveformTransformationTrimFade::FadeModeToFadeFunctionMap[EWaveEditorFadeMode::Linear]; }));
 
 	ToolkitCommands->MapAction(
 		Commands.FadeInExponential,
-		FExecuteAction::CreateLambda([this] { FadeInCurve = UWaveformTransformationTrimFade::FadeModeToCurveValueMap[EWaveEditorFadeMode::Exponetial]; }),
-		FCanExecuteAction::CreateLambda([this] {return true; }),
-		FIsActionChecked::CreateLambda([this] {return FadeInCurve == UWaveformTransformationTrimFade::FadeModeToCurveValueMap[EWaveEditorFadeMode::Exponetial]; }));
+		FExecuteAction::CreateLambda([this] 
+			{ 
+				FadeInCurve = UWaveformTransformationTrimFade::FadeModeToCurveValueMap[EWaveEditorFadeMode::Exponential];
+				FadeInType = UWaveformTransformationTrimFade::FadeModeToFadeFunctionMap[EWaveEditorFadeMode::Exponential];
+			}),
+		FCanExecuteAction::CreateLambda([this] { return true; }),
+		FIsActionChecked::CreateLambda([this] { return FadeInType == UWaveformTransformationTrimFade::FadeModeToFadeFunctionMap[EWaveEditorFadeMode::Exponential]; }));
 
 	ToolkitCommands->MapAction(
 		Commands.FadeInLogarithmic,
-		FExecuteAction::CreateLambda([this] { FadeInCurve = UWaveformTransformationTrimFade::FadeModeToCurveValueMap[EWaveEditorFadeMode::Logarithmic]; }),
-		FCanExecuteAction::CreateLambda([this] {return true; }),
-		FIsActionChecked::CreateLambda([this] {return FadeInCurve == UWaveformTransformationTrimFade::FadeModeToCurveValueMap[EWaveEditorFadeMode::Logarithmic]; }));
+		FExecuteAction::CreateLambda([this] 
+			{ 
+				FadeInCurve = UWaveformTransformationTrimFade::FadeModeToCurveValueMap[EWaveEditorFadeMode::Logarithmic];
+				FadeInType = UWaveformTransformationTrimFade::FadeModeToFadeFunctionMap[EWaveEditorFadeMode::Logarithmic];
+			}),
+		FCanExecuteAction::CreateLambda([this] { return true; }),
+		FIsActionChecked::CreateLambda([this] { return FadeInType == UWaveformTransformationTrimFade::FadeModeToFadeFunctionMap[EWaveEditorFadeMode::Logarithmic]; }));
 
 	ToolkitCommands->MapAction(
 		Commands.FadeInSigmoid,
-		FExecuteAction::CreateLambda([this] { FadeInCurve = UWaveformTransformationTrimFade::FadeModeToCurveValueMap[EWaveEditorFadeMode::Sigmoid]; }),
-		FCanExecuteAction::CreateLambda([this] {return true; }),
-		FIsActionChecked::CreateLambda([this] {return FadeInCurve == UWaveformTransformationTrimFade::FadeModeToCurveValueMap[EWaveEditorFadeMode::Sigmoid]; }));
+		FExecuteAction::CreateLambda([this] 
+			{ 
+				FadeInCurve = UWaveformTransformationTrimFade::FadeModeToCurveValueMap[EWaveEditorFadeMode::Sigmoid];
+				FadeInType = UWaveformTransformationTrimFade::FadeModeToFadeFunctionMap[EWaveEditorFadeMode::Sigmoid];
+			}),
+		FCanExecuteAction::CreateLambda([this] { return true; }),
+		FIsActionChecked::CreateLambda([this] { return FadeInType == UWaveformTransformationTrimFade::FadeModeToFadeFunctionMap[EWaveEditorFadeMode::Sigmoid]; }));
 
 	ToolkitCommands->MapAction(
 		Commands.ToggleFadeOut,
@@ -1412,27 +1529,43 @@ bool FWaveformEditor::BindCommands()
 
 	ToolkitCommands->MapAction(
 		Commands.FadeOutLinear,
-		FExecuteAction::CreateLambda([this] { FadeOutCurve = UWaveformTransformationTrimFade::FadeModeToCurveValueMap[EWaveEditorFadeMode::Linear]; }),
+		FExecuteAction::CreateLambda([this] 
+			{ 
+				FadeOutCurve = UWaveformTransformationTrimFade::FadeModeToCurveValueMap[EWaveEditorFadeMode::Linear];
+				FadeOutType = UWaveformTransformationTrimFade::FadeModeToFadeFunctionMap[EWaveEditorFadeMode::Linear];
+			}),
 		FCanExecuteAction::CreateLambda([this] {return true; }),
-		FIsActionChecked::CreateLambda([this] {return FadeOutCurve == UWaveformTransformationTrimFade::FadeModeToCurveValueMap[EWaveEditorFadeMode::Linear]; }));
+		FIsActionChecked::CreateLambda([this] {return FadeOutType == UWaveformTransformationTrimFade::FadeModeToFadeFunctionMap[EWaveEditorFadeMode::Linear]; }));
 
 	ToolkitCommands->MapAction(
 		Commands.FadeOutExponential,
-		FExecuteAction::CreateLambda([this] { FadeOutCurve = UWaveformTransformationTrimFade::FadeModeToCurveValueMap[EWaveEditorFadeMode::Exponetial]; }),
+		FExecuteAction::CreateLambda([this] 
+			{ 
+				FadeOutCurve = UWaveformTransformationTrimFade::FadeModeToCurveValueMap[EWaveEditorFadeMode::Exponential];
+				FadeOutType = UWaveformTransformationTrimFade::FadeModeToFadeFunctionMap[EWaveEditorFadeMode::Exponential];
+			}),
 		FCanExecuteAction::CreateLambda([this] {return true; }),
-		FIsActionChecked::CreateLambda([this] {return FadeOutCurve == UWaveformTransformationTrimFade::FadeModeToCurveValueMap[EWaveEditorFadeMode::Exponetial]; }));
+		FIsActionChecked::CreateLambda([this] {return FadeOutType == UWaveformTransformationTrimFade::FadeModeToFadeFunctionMap[EWaveEditorFadeMode::Exponential]; }));
 
 	ToolkitCommands->MapAction(
 		Commands.FadeOutLogarithmic,
-		FExecuteAction::CreateLambda([this] { FadeOutCurve = UWaveformTransformationTrimFade::FadeModeToCurveValueMap[EWaveEditorFadeMode::Logarithmic]; }),
+		FExecuteAction::CreateLambda([this] 
+			{ 
+				FadeOutCurve = UWaveformTransformationTrimFade::FadeModeToCurveValueMap[EWaveEditorFadeMode::Logarithmic];
+				FadeOutType = UWaveformTransformationTrimFade::FadeModeToFadeFunctionMap[EWaveEditorFadeMode::Logarithmic];
+			}),
 		FCanExecuteAction::CreateLambda([this] {return true; }),
-		FIsActionChecked::CreateLambda([this] {return FadeOutCurve == UWaveformTransformationTrimFade::FadeModeToCurveValueMap[EWaveEditorFadeMode::Logarithmic]; }));
+		FIsActionChecked::CreateLambda([this] {return FadeOutType == UWaveformTransformationTrimFade::FadeModeToFadeFunctionMap[EWaveEditorFadeMode::Logarithmic]; }));
 
 	ToolkitCommands->MapAction(
 		Commands.FadeOutSigmoid,
-		FExecuteAction::CreateLambda([this] { FadeOutCurve = UWaveformTransformationTrimFade::FadeModeToCurveValueMap[EWaveEditorFadeMode::Sigmoid]; }),
+		FExecuteAction::CreateLambda([this] 
+			{ 
+				FadeOutCurve = UWaveformTransformationTrimFade::FadeModeToCurveValueMap[EWaveEditorFadeMode::Sigmoid];
+				FadeOutType = UWaveformTransformationTrimFade::FadeModeToFadeFunctionMap[EWaveEditorFadeMode::Sigmoid];
+			}),
 		FCanExecuteAction::CreateLambda([this] {return true; }),
-		FIsActionChecked::CreateLambda([this] {return FadeOutCurve == UWaveformTransformationTrimFade::FadeModeToCurveValueMap[EWaveEditorFadeMode::Sigmoid]; }));
+		FIsActionChecked::CreateLambda([this] {return FadeOutType == UWaveformTransformationTrimFade::FadeModeToFadeFunctionMap[EWaveEditorFadeMode::Sigmoid]; }));
 
 	ToolkitCommands->MapAction(
 		Commands.LeftBoundsIncrease,
diff --git a/Engine/Plugins/Editor/WaveformEditor/Source/WaveformEditor/Public/WaveformEditor.h b/Engine/Plugins/Editor/WaveformEditor/Source/WaveformEditor/Public/WaveformEditor.h
--- a/Engine/Plugins/Editor/WaveformEditor/Source/WaveformEditor/Public/WaveformEditor.h
+++ b/Engine/Plugins/Editor/WaveformEditor/Source/WaveformEditor/Public/WaveformEditor.h
@@ -179,11 +179,13 @@ private:
 
 	static constexpr float DefaultFadeInAmount = 0.5f;
 	float CachedFadeInAmount = 0.f;
-	float FadeInCurve = UWaveformTransformationTrimFade::FadeModeToCurveValueMap[EWaveEditorFadeMode::Exponetial];
+	float FadeInCurve = UWaveformTransformationTrimFade::FadeModeToCurveValueMap[EWaveEditorFadeMode::Exponential];
+	TSubclassOf<UFadeFunction> FadeInType = UWaveformTransformationTrimFade::FadeModeToFadeFunctionMap[EWaveEditorFadeMode::Exponential];
 
 	static constexpr float DefaultFadeOutAmount = 0.5f;
 	float CachedFadeOutAmount = 0.f;
-	float FadeOutCurve = UWaveformTransformationTrimFade::FadeModeToCurveValueMap[EWaveEditorFadeMode::Exponetial];
+	float FadeOutCurve = UWaveformTransformationTrimFade::FadeModeToCurveValueMap[EWaveEditorFadeMode::Exponential];
+	TSubclassOf<UFadeFunction> FadeOutType = UWaveformTransformationTrimFade::FadeModeToFadeFunctionMap[EWaveEditorFadeMode::Exponential];
 
 	// Tracking bool to ensure we detect wave cue changes and regenerate transformations
 	// Checking the event is bound may have issues if other systems need to subscribe for wave cue array changes
diff --git a/Engine/Plugins/Editor/WaveformEditor/Source/WaveformTransformations/Private/WaveformTransformationTrimFade.cpp b/Engine/Plugins/Editor/WaveformEditor/Source/WaveformTransformations/Private/WaveformTransformationTrimFade.cpp
--- a/Engine/Plugins/Editor/WaveformEditor/Source/WaveformTransformations/Private/WaveformTransformationTrimFade.cpp
+++ b/Engine/Plugins/Editor/WaveformEditor/Source/WaveformTransformations/Private/WaveformTransformationTrimFade.cpp
@@ -5,13 +5,32 @@
 
 #include UE_INLINE_GENERATED_CPP_BY_NAME(WaveformTransformationTrimFade)
 
+namespace
+{
+	// Makes the two functions for SCurve meet at (0.5, 0.5)
+	const double XMultiplier = FMath::Sqrt(10.0) / 5.0;
+}
+
+const TMap<EWaveEditorFadeMode, TSubclassOf<UFadeFunction>> UWaveformTransformationTrimFade::FadeModeToFadeFunctionMap =
+{
+	{EWaveEditorFadeMode::Linear, UFadeFunctionLinear::StaticClass()},
+	{EWaveEditorFadeMode::Exponential, UFadeCurveFunctionExponential::StaticClass()},
+	{EWaveEditorFadeMode::Logarithmic, UFadeCurveFunctionLogarithmic::StaticClass()},
+	{EWaveEditorFadeMode::Sigmoid, UFadeCurveFunctionSigmoid::StaticClass()}
+};
+
+PRAGMA_DISABLE_DEPRECATION_WARNINGS
 const TMap<EWaveEditorFadeMode, float> UWaveformTransformationTrimFade::FadeModeToCurveValueMap =
 {
 	{EWaveEditorFadeMode::Linear, 1.0f},
-	{EWaveEditorFadeMode::Exponetial, 3.0f},
+	{EWaveEditorFadeMode::Exponential, 3.0f},
 	{EWaveEditorFadeMode::Logarithmic, 0.25f},
-	{EWaveEditorFadeMode::Sigmoid, -0.1f}
+	{EWaveEditorFadeMode::Sigmoid, UWaveformTransformationTrimFade::MinFadeCurve}
 };
+PRAGMA_ENABLE_DEPRECATION_WARNINGS
+
+const float UWaveformTransformationTrimFade::MinFadeCurve = -0.1f;
+const float UWaveformTransformationTrimFade::MaxFadeCurve = 10.f;
 
 namespace WaveformTransformationTrimFadeNames
 {
@@ -19,22 +38,30 @@ namespace WaveformTransformationTrimFadeNames
 	static FLazyName EndTimeName(GET_MEMBER_NAME_CHECKED(UWaveformTransformationTrimFade, EndTime));
 }
 
-static void ApplyFadeIn(Audio::FAlignedFloatBuffer& InputAudio, const float FadeLength, const float FadeCurve, const float SCurveSharpness, const int32 NumChannels, const float SampleRate)
+static void ApplyFadeIn(Audio::FAlignedFloatBuffer& InputAudio, const FFadeFunctionData& FadeFunctions, const float FadeLength, const float FadeCurve, const float SCurveSharpness, const int32 NumChannels, const float SampleRate)
 {
 	check(NumChannels > 0);
+	check(SampleRate > 0.f);
+	const float FadeDuration = FadeFunctions.FadeIn ? FadeFunctions.FadeIn->Duration : FadeLength;
+	const int32 FadeNumFrames = FMath::Min(FadeDuration * SampleRate, InputAudio.Num() / NumChannels);
 
-	if (InputAudio.Num() < NumChannels || FadeLength < SMALL_NUMBER)
+	if (InputAudio.Num() < NumChannels || FadeDuration < SMALL_NUMBER || FadeNumFrames == 0)
 	{
 		return;
 	}
 
-	const int32 FadeNumFrames = FMath::Min(FadeLength * SampleRate, InputAudio.Num() / NumChannels);
 	float* InputPtr = InputAudio.GetData();
 
 	for (int32 FrameIndex = 0; FrameIndex < FadeNumFrames; ++FrameIndex)
 	{
-		const float FadeFraction = (float)FrameIndex / FadeNumFrames;
-		const float EnvValue = UWaveformTransformationTrimFade::GetFadeInCurveValue(FadeCurve, FadeFraction, SCurveSharpness);
+		const float FadeFraction = static_cast<float>(FrameIndex) / static_cast<float>(FadeNumFrames);
+		check(FadeFraction <= 1.0f);
+		check(FadeFraction >= 0.0f);
+
+		PRAGMA_DISABLE_DEPRECATION_WARNINGS
+		const float EnvValue = FadeFunctions.FadeIn ? FadeFunctions.FadeIn->GetFadeInCurveValue(FadeFraction) :
+		UWaveformTransformationTrimFade::GetFadeInCurveValue(FadeCurve, FadeFraction, SCurveSharpness);
+		PRAGMA_ENABLE_DEPRECATION_WARNINGS
 
 		for (int32 ChannelIt = 0; ChannelIt < NumChannels; ++ChannelIt)
 		{
@@ -45,23 +72,31 @@ static void ApplyFadeIn(Audio::FAlignedFloatBuffer& InputAudio, const float Fade
 	}
 }
 
-static void ApplyFadeOut(Audio::FAlignedFloatBuffer& InputAudio, const float FadeLength, const float FadeCurve, const float SCurveSharpness, const int32 NumChannels, const float SampleRate)
+static void ApplyFadeOut(Audio::FAlignedFloatBuffer& InputAudio, const FFadeFunctionData& FadeFunctions, const float FadeLength, const float FadeCurve, const float SCurveSharpness, const int32 NumChannels, const float SampleRate)
 {
 	check(NumChannels > 0);
+	check(SampleRate > 0.f);
+	const float FadeDuration = FadeFunctions.FadeOut ? FadeFunctions.FadeOut->Duration : FadeLength;
+	const int32 FadeNumFrames = FMath::Min(FadeDuration * SampleRate, InputAudio.Num() / NumChannels);
 
-	if (InputAudio.Num() < NumChannels || FadeLength < SMALL_NUMBER)
+	if (InputAudio.Num() < NumChannels || FadeDuration < SMALL_NUMBER || FadeNumFrames == 0)
 	{
 		return;
 	}
 
-	const int32 FadeNumFrames = FMath::Min(FadeLength * SampleRate, InputAudio.Num() / NumChannels);
 	const int32 StartSampleIndex = InputAudio.Num() - (FadeNumFrames * NumChannels);
 	float* InputPtr = &InputAudio[StartSampleIndex];
 
 	for (int32 FrameIndex = 0; FrameIndex < FadeNumFrames; ++FrameIndex)
 	{
-		const float FadeFraction = (float)FrameIndex / FadeNumFrames;
-		float EnvValue = UWaveformTransformationTrimFade::GetFadeOutCurveValue(FadeCurve, FadeFraction, SCurveSharpness);
+		const float FadeFraction = static_cast<float>(FrameIndex) / static_cast<float>(FadeNumFrames);
+		check(FadeFraction <= 1.0f);
+		check(FadeFraction >= 0.0f);
+
+		PRAGMA_DISABLE_DEPRECATION_WARNINGS
+		const float EnvValue = FadeFunctions.FadeOut ? FadeFunctions.FadeOut->GetFadeOutCurveValue(FadeFraction) :
+		UWaveformTransformationTrimFade::GetFadeOutCurveValue(FadeCurve, FadeFraction, SCurveSharpness);
+		PRAGMA_ENABLE_DEPRECATION_WARNINGS
 
 		for (int32 ChannelIt = 0; ChannelIt < NumChannels; ++ChannelIt)
 		{
@@ -71,16 +106,18 @@ static void ApplyFadeOut(Audio::FAlignedFloatBuffer& InputAudio, const float Fad
 		InputPtr += NumChannels;
 	}
 }
-
-FWaveTransformationTrimFade::FWaveTransformationTrimFade(double InStartTime, double InEndTime, float InStartFadeTime, float InStartFadeCurve, float InStartSCurveSharpness, float InEndFadeTime, float InEndFadeCurve, float InEndSCurveSharpness)
-	: StartTime(InStartTime)
+PRAGMA_DISABLE_DEPRECATION_WARNINGS
+FWaveTransformationTrimFade::FWaveTransformationTrimFade(const FFadeFunctionData& InFadeFunctions, double InStartTime, double InEndTime, float InStartFadeTime, float InStartFadeCurve, float InStartSCurveSharpness, float InEndFadeTime, float InEndFadeCurve, float InEndSCurveSharpness)
+	: FadeFunctions(InFadeFunctions)
+	, StartTime(InStartTime)
 	, EndTime(InEndTime)
 	, StartFadeTime(InStartFadeTime)
-	, StartFadeCurve(FMath::Max(InStartFadeCurve, -0.1f))
+	, StartFadeCurve(FMath::Max(InStartFadeCurve, UWaveformTransformationTrimFade::MinFadeCurve))
 	, StartSCurveSharpness(InStartSCurveSharpness)
 	, EndFadeTime(InEndFadeTime)
-	, EndFadeCurve(FMath::Max(InEndFadeCurve, -0.1f))
+	, EndFadeCurve(FMath::Max(InEndFadeCurve, UWaveformTransformationTrimFade::MinFadeCurve))
 	, EndSCurveSharpness(InEndSCurveSharpness){}
+PRAGMA_ENABLE_DEPRECATION_WARNINGS
 
 void FWaveTransformationTrimFade::ProcessAudio(Audio::FWaveformTransformationWaveInfo& InOutWaveInfo) const
 {
@@ -136,16 +173,17 @@ void FWaveTransformationTrimFade::ProcessAudio(Audio::FWaveformTransformationWav
 
 	FMemory::Memcpy(InputAudio.GetData(), TempBuffer.GetData(), FinalSize * sizeof(float));
 
-
-	if(StartFadeTime > 0.f)
+	PRAGMA_DISABLE_DEPRECATION_WARNINGS
+	if ((FadeFunctions.FadeIn && FadeFunctions.FadeIn->Duration > 0.f) || StartFadeTime > 0.f)
 	{
-		ApplyFadeIn(InputAudio, StartFadeTime, StartFadeCurve, StartSCurveSharpness, InOutWaveInfo.NumChannels, InOutWaveInfo.SampleRate);
+		ApplyFadeIn(InputAudio, FadeFunctions, StartFadeTime, StartFadeCurve, StartSCurveSharpness, InOutWaveInfo.NumChannels, InOutWaveInfo.SampleRate);
 	}
-	
-	if(EndFadeTime > 0.f)
+
+	if ((FadeFunctions.FadeOut && FadeFunctions.FadeOut->Duration > 0.f) || EndFadeTime > 0.f)
 	{
-		ApplyFadeOut(InputAudio, EndFadeTime, EndFadeCurve, EndSCurveSharpness, InOutWaveInfo.NumChannels, InOutWaveInfo.SampleRate);
-	}
+		ApplyFadeOut(InputAudio, FadeFunctions, EndFadeTime, EndFadeCurve, EndSCurveSharpness, InOutWaveInfo.NumChannels, InOutWaveInfo.SampleRate);
+	}	
+	PRAGMA_ENABLE_DEPRECATION_WARNINGS
 }
 
 void UWaveformTransformationTrimFade::PostEditChangeProperty(FPropertyChangedEvent& PropertyChangedEvent)
@@ -153,13 +191,28 @@ void UWaveformTransformationTrimFade::PostEditChangeProperty(FPropertyChangedEve
 	Super::PostEditChangeProperty(PropertyChangedEvent);
 
 	const float MaxDuration = EndTime - StartTime;
+
+	PRAGMA_DISABLE_DEPRECATION_WARNINGS
 	EndFadeTime = FMath::Min(EndFadeTime, MaxDuration);
 	StartFadeTime = FMath::Min(StartFadeTime, MaxDuration);
+	PRAGMA_ENABLE_DEPRECATION_WARNINGS
+
+	if (FadeFunctions.FadeIn)
+	{
+		FadeFunctions.FadeIn->Duration = FMath::Min(FadeFunctions.FadeIn->Duration, MaxDuration);
+	}
+
+	if (FadeFunctions.FadeOut)
+	{
+		FadeFunctions.FadeOut->Duration = FMath::Min(FadeFunctions.FadeOut->Duration, MaxDuration);
+	}
 }
 
 Audio::FTransformationPtr UWaveformTransformationTrimFade::CreateTransformation() const
 {
-	return MakeUnique<FWaveTransformationTrimFade>(StartTime, EndTime, StartFadeTime, StartFadeCurve, StartSCurveSharpness, EndFadeTime, EndFadeCurve, EndSCurveSharpness);
+	PRAGMA_DISABLE_DEPRECATION_WARNINGS
+	return MakeUnique<FWaveTransformationTrimFade>(FadeFunctions, StartTime, EndTime, StartFadeTime, StartFadeCurve, StartSCurveSharpness, EndFadeTime, EndFadeCurve, EndSCurveSharpness);
+	PRAGMA_ENABLE_DEPRECATION_WARNINGS
 }
 
 void UWaveformTransformationTrimFade::UpdateConfiguration(FWaveTransformUObjectConfiguration& InOutConfiguration)
@@ -176,15 +229,17 @@ const double UWaveformTransformationTrimFade::GetFadeInCurveValue(const float St
 
 	if (StartFadeCurve < 0)
 	{
-		double Slope = 10 * SCurveSharpness * (PI + 1);
-		
+		const double X = FadeFraction;
+		const double A = static_cast<float>(-FMath::Clamp(SCurveSharpness, 0.1f, 1.0f));
+		const double K = 20 * FMath::Pow(A, 5.0);
+
 		if (FadeFraction <= 0.5)
 		{
-			CurveValue = (((FMath::Tanh(((Slope * FadeFraction * PI) / 2) - ((Slope * PI) / 4))) + 1) / 2) * (2 * FadeFraction);
+			CurveValue = -(FMath::Tanh(K * PI * FMath::LogX(10, XMultiplier * X) + ((K * PI) / 2))) / 2 + 0.5;
 		}
 		else
 		{
-			CurveValue = ((FMath::Tanh((-Slope * (-FadeFraction + 1) * PI) / 2 + ((Slope * PI) / 4)) - 1) / 2) * (2 * (-FadeFraction + 1)) + 1;
+			CurveValue = (FMath::Tanh(K * PI * FMath::LogX(10, XMultiplier * (-X + 1)) + ((K * PI) / 2))) / 2 + 0.5;
 		}
 	}
 	else
@@ -198,18 +253,20 @@ const double UWaveformTransformationTrimFade::GetFadeInCurveValue(const float St
 const double UWaveformTransformationTrimFade::GetFadeOutCurveValue(const float EndFadeCurve, const double FadeFraction, const float SCurveSharpness)
 {
 	double CurveValue = 1;
-
+	
 	if (EndFadeCurve < 0)
 	{
-		double Slope = 10 * SCurveSharpness * (PI + 1);
-		
+		const double X = FadeFraction;
+		const double A = static_cast<float>(FMath::Clamp(SCurveSharpness, 0.1f, 1.0f));
+		const double K = 20 * FMath::Pow(A, 5.0);
+
 		if (FadeFraction <= 0.5)
 		{
-			CurveValue = -(((FMath::Tanh(((Slope * FadeFraction * PI) / 2) - ((Slope * PI) / 4))) + 1) / 2) * (2 * FadeFraction) + 1;
+			CurveValue = -(FMath::Tanh(K * PI * FMath::LogX(10, XMultiplier * X) + ((K * PI) / 2))) / 2 + 0.5;
 		}
 		else
 		{
-			CurveValue = -((FMath::Tanh((-Slope * (-FadeFraction + 1) * PI) / 2 + ((Slope * PI) / 4)) - 1) / 2) * (2 * (-FadeFraction + 1));
+			CurveValue = (FMath::Tanh(K * PI * FMath::LogX(10, XMultiplier * (-X + 1)) + ((K * PI) / 2))) / 2 + 0.5;
 		}
 	}
 	else
@@ -229,3 +286,111 @@ void UWaveformTransformationTrimFade::UpdateDurationProperties(const float InAva
 	StartTime = FMath::Clamp(StartTime, 0.f, AvailableWaveformDuration - UE_KINDA_SMALL_NUMBER);
 	EndTime = EndTime < 0 ? EndTime = AvailableWaveformDuration : FMath::Clamp(EndTime, StartTime + UE_KINDA_SMALL_NUMBER, AvailableWaveformDuration);
 }
+
+const double UFadeFunctionLinear::GetFadeInCurveValue(const double FadeFraction) const
+{
+	check(FadeFraction <= 1.f);
+	check(FadeFraction >= 0.f);
+
+	return FadeFraction;
+}
+
+const double UFadeFunctionLinear::GetFadeOutCurveValue(const double FadeFraction) const
+{
+	check(FadeFraction <= 1.f);
+	check(FadeFraction >= 0.f);
+
+	return 1.0 - FadeFraction;
+}
+
+const double UFadeCurveFunctionExponential::GetFadeInCurveValue(const double FadeFraction) const
+{
+	check(FadeFraction <= 1.f);
+	check(FadeFraction >= 0.f);
+
+	return FMath::Pow(FadeFraction, FadeCurve);
+}
+
+const double UFadeCurveFunctionExponential::GetFadeOutCurveValue(const double FadeFraction) const
+{
+	// Pow functions return NaN for negative bases with non-integer exponents
+	check(FadeFraction <= 1.f);
+	check(FadeFraction >= 0.f);
+
+	return FMath::Pow(-FadeFraction + 1.0, FadeCurve);
+}
+
+const double UFadeCurveFunctionLogarithmic::GetFadeInCurveValue(const double FadeFraction) const
+{
+	check(FadeFraction <= 1.f);
+	check(FadeFraction >= 0.f);
+
+	return FMath::Pow(FadeFraction, FadeCurve);
+}
+
+const double UFadeCurveFunctionLogarithmic::GetFadeOutCurveValue(const double FadeFraction) const
+{
+	// Pow functions return NaN for negative bases with non-integer exponents
+	check(FadeFraction <= 1.f);
+	check(FadeFraction >= 0.f);
+
+	return FMath::Pow(-FadeFraction + 1.0, FadeCurve);
+}
+
+const double UFadeCurveFunctionSigmoid::GetFadeInCurveValue(const double FadeFraction) const
+{
+	// LogX will produce NaN if FadeFraction is <= 0 or is 1
+	check(FadeFraction <= 1.f);
+	check(FadeFraction >= 0.f);
+	check(XMultiplier > 0.0);
+
+	if (FadeFraction == 0.f || FadeFraction == 1.f)
+	{
+		return FadeFraction;
+	}
+
+	double CurveValue = 1;
+	const double x = FadeFraction;
+	const double a = static_cast<float>(-FMath::Clamp(SFadeCurve, 0.1f, 1.0f));
+	const double k = 20 * FMath::Pow(a, 5);
+
+	if (FadeFraction <= 0.5)
+	{
+		CurveValue = -(FMath::Tanh(k * PI * FMath::LogX(10, XMultiplier * x) + ((k * PI) / 2))) / 2 + 0.5;
+	}
+	else
+	{
+		CurveValue = (FMath::Tanh(k * PI * FMath::LogX(10, XMultiplier * (-x + 1)) + ((k * PI) / 2))) / 2 + 0.5;
+	}
+
+	return CurveValue;
+}
+
+const double UFadeCurveFunctionSigmoid::GetFadeOutCurveValue(const double FadeFraction) const
+{
+	// LogX will produce NaN if FadeFraction is <= 0 or is 1
+	check(FadeFraction <= 1.f);
+	check(FadeFraction >= 0.f);
+	check(XMultiplier > 0.0);
+
+	if (FadeFraction == 0.f || FadeFraction == 1.f)
+	{
+		return 1.0 - FadeFraction;
+	}
+
+	double CurveValue = 1;
+	const double x = FadeFraction;
+	const double a = static_cast<float>(FMath::Clamp(SFadeCurve, 0.1f, 1.0f));
+	const double k = 20 * FMath::Pow(a, 5);
+
+	if (FadeFraction <= 0.5)
+	{
+		CurveValue = -(FMath::Tanh(k * PI * FMath::LogX(10, XMultiplier * x) + ((k * PI) / 2))) / 2 + 0.5;
+	}
+	else
+	{
+		CurveValue = (FMath::Tanh(k * PI * FMath::LogX(10, XMultiplier * (-x + 1)) + ((k * PI) / 2))) / 2 + 0.5;
+	}
+
+	return CurveValue;
+}
diff --git a/Engine/Plugins/Editor/WaveformEditor/Source/WaveformTransformations/Public/WaveformTransformationTrimFade.h b/Engine/Plugins/Editor/WaveformEditor/Source/WaveformTransformations/Public/WaveformTransformationTrimFade.h
--- a/Engine/Plugins/Editor/WaveformEditor/Source/WaveformTransformations/Public/WaveformTransformationTrimFade.h
+++ b/Engine/Plugins/Editor/WaveformEditor/Source/WaveformTransformations/Public/WaveformTransformationTrimFade.h
@@ -9,11 +9,119 @@
 
 #define UE_API WAVEFORMTRANSFORMATIONS_API
 
+UCLASS(MinimalAPI, Abstract, EditInlineNew, DefaultToInstanced, CollapseCategories, NotBlueprintable)
+class UFadeFunction : public UObject
+{
+	GENERATED_BODY()
+public:
+
+	virtual const double GetFadeInCurveValue(const double FadeFraction) const { return FadeFraction; }
+	virtual const double GetFadeOutCurveValue(const double FadeFraction) const { return FadeFraction; }
+
+	UPROPERTY(EditAnywhere, Category = "Fade", BlueprintReadWrite, meta = (ClampMin = 0.0, DisplayName = "Fade Duration"))
+	float Duration = 0.5f;
+};
+
+UCLASS(MinimalAPI, Abstract, EditInlineNew, DefaultToInstanced, CollapseCategories, NotBlueprintable)
+class UFadeCurveFunction : public UFadeFunction
+{
+	GENERATED_BODY()
+public:
+
+	virtual const double GetFadeInCurveValue(const double FadeFraction) const override { return FadeFraction; }
+	virtual const double GetFadeOutCurveValue(const double FadeFraction) const override { return FadeFraction; }
+
+	virtual void SetFadeCurve(const float InFadeCurve) {}
+	virtual const float GetFadeCurve() const { return 0.f; }
+	virtual const FText GetFadeCurvePropertyName() const { return FText::FromName(NAME_None); }
+
+protected:
+	virtual const float GetFadeCurveMin() const { return 0.f; }
+	virtual const float GetFadeCurveMax() const { return 1.f; }
+};
+
+UCLASS(MinimalAPI, NotBlueprintable, meta = (DisplayName = "Linear"))
+class UFadeFunctionLinear : public UFadeFunction
+{
+	GENERATED_BODY()
+public:
+	
+	UE_API virtual const double GetFadeInCurveValue(const double FadeFraction) const override;
+	UE_API virtual const double GetFadeOutCurveValue(const double FadeFraction) const override;
+};
+
+UCLASS(MinimalAPI, NotBlueprintable, meta = (DisplayName = "Exponential"))
+class UFadeCurveFunctionExponential : public UFadeCurveFunction
+{
+	GENERATED_BODY()
+public:
+
+	UE_API virtual const double GetFadeInCurveValue(const double FadeFraction) const override;
+	UE_API virtual const double GetFadeOutCurveValue(const double FadeFraction) const override;
+
+	virtual void SetFadeCurve(const float InFadeCurve) override { FadeCurve = FMath::Clamp(InFadeCurve, GetFadeCurveMin(), GetFadeCurveMax()); }
+	virtual const float GetFadeCurve() const override { return FadeCurve; }
+	virtual const FText GetFadeCurvePropertyName() const override { return FText::FromName(GET_MEMBER_NAME_CHECKED(UFadeCurveFunctionExponential, FadeCurve)); }
+
+	UPROPERTY(EditAnywhere, Category = "Fade", meta = (ClampMin = 1, ClampMax = 10.0, DisplayName = "Fade Curve"))
+	float FadeCurve = 3.f;
+
+protected:
+	virtual const float GetFadeCurveMin() const override { return 1.f; }
+	virtual const float GetFadeCurveMax() const override { return 10.f; }
+};
+
+UCLASS(MinimalAPI, NotBlueprintable, meta = (DisplayName = "Logarithmic"))
+class UFadeCurveFunctionLogarithmic : public UFadeCurveFunction
+{
+	GENERATED_BODY()
+public:
+
+	UE_API virtual const double GetFadeInCurveValue(const double FadeFraction) const override;
+	UE_API virtual const double GetFadeOutCurveValue(const double FadeFraction) const override;
+
+	virtual void SetFadeCurve(const float InFadeCurve) override { FadeCurve = FMath::Clamp(InFadeCurve, GetFadeCurveMin(), GetFadeCurveMax()); }
+	virtual const float GetFadeCurve() const override { return FadeCurve; }
+	virtual const FText GetFadeCurvePropertyName() const override { return FText::FromName(GET_MEMBER_NAME_CHECKED(UFadeCurveFunctionLogarithmic, FadeCurve)); }
+
+	UPROPERTY(EditAnywhere, Category = "Fade", meta = (ClampMin = 0.0, ClampMax = 1, DisplayName = "Fade Curve"))
+	float FadeCurve = 0.25f;
+};
+
+UCLASS(MinimalAPI, NotBlueprintable, meta = (DisplayName = "Sigmoid"))
+class UFadeCurveFunctionSigmoid : public UFadeCurveFunction
+{
+	GENERATED_BODY()
+public:
+
+	UE_API virtual const double GetFadeInCurveValue(const double FadeFraction) const override;
+	UE_API virtual const double GetFadeOutCurveValue(const double FadeFraction) const override;
+
+	virtual void SetFadeCurve(const float InFadeCurve) override { SFadeCurve = FMath::Clamp(InFadeCurve, GetFadeCurveMin(), GetFadeCurveMax()); }
+	virtual const float GetFadeCurve() const override { return SFadeCurve; }
+	virtual const FText GetFadeCurvePropertyName() const override { return FText::FromName(GET_MEMBER_NAME_CHECKED(UFadeCurveFunctionSigmoid, SFadeCurve)); }
+
+	UPROPERTY(EditAnywhere, Category = "Fade", meta = (ClampMin = 0.0, ClampMax = 1, DisplayName = "S-Fade Curve"))
+	float SFadeCurve = 0.6f;
+};
+
+USTRUCT()
+struct FFadeFunctionData
+{
+	GENERATED_BODY()
+
+	UPROPERTY(EditAnywhere, Category = "Fade", Instanced)
+	TObjectPtr<UFadeFunction> FadeIn;
+
+	UPROPERTY(EditAnywhere, Category = "Fade", Instanced)
+	TObjectPtr<UFadeFunction> FadeOut;
+};
+
 UENUM()
 enum class EWaveEditorFadeMode : uint8
 {
 	Linear = 0,
-	Exponetial,
+	Exponential,
 	Logarithmic,
 	Sigmoid
 };
@@ -21,12 +129,15 @@ enum class EWaveEditorFadeMode : uint8
 class FWaveTransformationTrimFade : public Audio::IWaveTransformation
 {
 public:
-	UE_API explicit FWaveTransformationTrimFade(double InStartTime, double InEndTime, float InStartFadeTime, float InStartFadeCurve, float InStartSCurveSharpness, float InEndFadeTime, float InEndFadeCurve, float InEndSCurveSharpness);
+	UE_API explicit FWaveTransformationTrimFade(const FFadeFunctionData& InFadeFunctions, double InStartTime, double InEndTime, float InStartFadeTime, float InStartFadeCurve, float InStartSCurveSharpness, float InEndFadeTime, float InEndFadeCurve, float InEndSCurveSharpness);
 	UE_API virtual void ProcessAudio(Audio::FWaveformTransformationWaveInfo& InOutWaveInfo) const override;
 
 	virtual constexpr Audio::ETransformationPriority FileChangeLengthPriority() const override { return Audio::ETransformationPriority::High; }
 
 private:
+
+	FFadeFunctionData FadeFunctions;
+
 	double StartTime = 0.0;
 	double EndTime = 0.0;
 
@@ -55,34 +166,71 @@ public:
 	UPROPERTY(EditAnywhere, Category = "Trim")
 	double EndTime = -1.0;
 
-	// Relative to StartTime
+	UPROPERTY(EditAnywhere, Category = "Fade", meta = (ClampMin = 0.0, DisplayName = "Fade Functions"))
+	FFadeFunctionData FadeFunctions;
+
+	UE_DEPRECATED(5.7, "Individual fade properties have been deprecated. Use FadeFunctions instead.")
+	/*
+	 * Relative to StartTime
+	 * Individual fade properties have been deprecated. Use FadeFunctions instead.
+	 */ 
 	UPROPERTY(EditAnywhere, Category = "Fade", meta = (ClampMin = 0.0, DisplayName = "Fade-In Duration"))
 	float StartFadeTime = 0.f;
-
+	
+	UE_DEPRECATED(5.7, "Individual fade properties have been deprecated. Use FadeFunctions instead.")
+	/*
+	 * Individual fade properties have been deprecated. Use FadeFunctions instead.
+	 */
 	UPROPERTY(EditAnywhere, Category = "Fade", meta=(ClampMin = -0.1, ClampMax = 10.0, DisplayName = "Fade-In Curve"))
 	float StartFadeCurve = 1.f;
 
-	UPROPERTY(EditAnywhere, Category = "Fade", meta = (ClampMin = -1.0, ClampMax = 1.0, DisplayName = "S-Curve Sharpness", EditCondition = "StartFadeCurve < 0", EditConditionHides))
-	float StartSCurveSharpness = 0.1f;
-
-	// Relative to EndTime. Fade out starts at EndTime - EndFadeTime and ends at EndTime
+	UE_DEPRECATED(5.7, "Individual fade properties have been deprecated. Use FadeFunctions instead.")
+	/*
+	 * Individual fade properties have been deprecated. Use FadeFunctions instead.
+	 */
+	UPROPERTY(EditAnywhere, Category = "Fade", meta = (ClampMin = 0, ClampMax = 1.0, DisplayName = "S-Curve Sharpness", EditCondition = "StartFadeCurve < 0", EditConditionHides))
+	float StartSCurveSharpness = 0.75f;
+
+	UE_DEPRECATED(5.7, "Individual fade properties have been deprecated. Use FadeFunctions instead.")
+	/*
+	 * Relative to EndTime. Fade out starts at EndTime - EndFadeTime and ends at EndTime
+	 * Individual fade properties have been deprecated. Use FadeFunctions instead.
+	 */
 	UPROPERTY(EditAnywhere, Category = "Fade", meta = (ClampMin = 0.0, DisplayName = "Fade-Out Duration"))
 	float EndFadeTime = 0.f;
 
+	UE_DEPRECATED(5.7, "Individual fade properties have been deprecated. Use FadeFunctions instead.")
+	/*
+	 * Individual fade properties have been deprecated. Use FadeFunctions instead.
+	 */
 	UPROPERTY(EditAnywhere, Category = "Fade", meta=(ClampMin = -0.1, ClampMax = 10.0, DisplayName = "Fade-Out Curve"))
 	float EndFadeCurve = 1.f;
 
-	UPROPERTY(EditAnywhere, Category = "Fade", meta = (ClampMin = -1.0, ClampMax = 1.0, DisplayName = "S-Curve Sharpness", EditCondition = "EndFadeCurve < 0", EditConditionHides))
-	float EndSCurveSharpness = 0.1;
+	UE_DEPRECATED(5.7, "Individual fade properties have been deprecated. Use FadeFunctions instead.")
+	/*
+	 * Individual fade properties have been deprecated. Use FadeFunctions instead.
+	 */
+	UPROPERTY(EditAnywhere, Category = "Fade", meta = (ClampMin = 0, ClampMax = 1.0, DisplayName = "S-Curve Sharpness", EditCondition = "EndFadeCurve < 0", EditConditionHides))
+	float EndSCurveSharpness = 0.75;
+
+	UE_DEPRECATED(5.7, "Individual fade properties have been deprecated. Use FadeFunctions instead.")
+	static UE_API const float MinFadeCurve;
+
+	UE_DEPRECATED(5.7, "Individual fade properties have been deprecated. Use FadeFunctions instead.")
+	static UE_API const float MaxFadeCurve;
 	
 	static UE_API const TMap<EWaveEditorFadeMode, float> FadeModeToCurveValueMap;
+	static UE_API const TMap<EWaveEditorFadeMode, TSubclassOf<UFadeFunction>> FadeModeToFadeFunctionMap;
 
 	UE_API virtual Audio::FTransformationPtr CreateTransformation() const override;
 
 	UE_API void UpdateConfiguration(FWaveTransformUObjectConfiguration& InOutConfiguration) override;
 	virtual constexpr Audio::ETransformationPriority GetTransformationPriority() const { return Audio::ETransformationPriority::High; }
 
+	UE_DEPRECATED(5.7, "Individual fade properties have been deprecated. Use FadeFunctions::FadeIn::GetFadeInCurveValue instead.")
 	static UE_API const double GetFadeInCurveValue(const float StartFadeCurve, const double FadeFraction, const float SCurveSharpness = 0);
+
+	UE_DEPRECATED(5.7, "Individual fade properties have been deprecated. Use FadeFunctions::FadeOut::GetFadeOutCurveValue instead.")
 	static UE_API const double GetFadeOutCurveValue(const float EndFadeCurve, const double FadeFraction, const float SCurveSharpness = 0);
 	
 private:	
@@ -90,52 +238,4 @@ private:
 	float AvailableWaveformDuration = -1.f;
 };
 
-// class UBaseClass;
-//
-// namespace  MyNamespace
-// {
-// 	class IBaseInterface;
-// }
-
-
-
-// // no good
-// #include "IBaseInterface.h" // in other module but that module is listed in the build.cs so we good there
-// #include "DerivedClass.generated.h"
-//
-// class PLUGIN_API FDerivedClass : public MyNamespace::IBaseInterface // error: Class 'MyNamespace' Not Found
-// {
-// public:
-// 	virtual void DoThing() const;
-// };
-//
-// UCLASS()
-// class PLUGIN_API UDerivedClass : public UBaseClass
-// {
-// 	GENERATED_BODY()
-// public:
-// 	TSharedPtr<MyNamespace::IBaseInterface> CreateFObject();
-// };
-//
-// // fine
-// #include "IBaseInterface.h"
-// #include "DerivedClass.generated.h"
-//
-// namespace MyNamespace
-// {
-// 	class PLUGIN_API FDerivedClass : public MyNamespace::IBaseInterface
-// 	{
-// 	public:
-// 		virtual void DoThing() const;
-// 	};
-// }
-//
-// UCLASS()
-// class PLUGIN_API UDerivedClass : public UBaseClass
-// {
-// 	GENERATED_BODY()
-// public:
-// 	TSharedPtr<MyNamespace::IBaseInterface> CreateFObject();
-// };
-
 #undef UE_API
diff --git a/Engine/Plugins/Editor/WaveformEditor/Source/WaveformTransformationsWidgets/Private/WaveformTransformationTrimFadeRenderer.cpp b/Engine/Plugins/Editor/WaveformEditor/Source/WaveformTransformationsWidgets/Private/WaveformTransformationTrimFadeRenderer.cpp
--- a/Engine/Plugins/Editor/WaveformEditor/Source/WaveformTransformationsWidgets/Private/WaveformTransformationTrimFadeRenderer.cpp
+++ b/Engine/Plugins/Editor/WaveformEditor/Source/WaveformTransformationsWidgets/Private/WaveformTransformationTrimFadeRenderer.cpp
@@ -147,10 +147,12 @@ int32 FWaveformTransformationTrimFadeRenderer::DrawFadeCurves(const FGeometry& A
 	return LayerId;
 }
 
+PRAGMA_DISABLE_DEPRECATION_WARNINGS
 void FWaveformTransformationTrimFadeRenderer::GenerateFadeCurves(const FGeometry& AllottedGeometry)
 {
 	check(StrongTrimFade);
-	const float FadeInFrames = StrongTrimFade->StartFadeTime * TransformationWaveInfo.SampleRate;
+	const float FadeInTime = StrongTrimFade->FadeFunctions.FadeIn ? StrongTrimFade->FadeFunctions.FadeIn->Duration : StrongTrimFade->StartFadeTime;
+	const float FadeInFrames = FadeInTime * TransformationWaveInfo.SampleRate;
 	const uint32 FadeInPixelLength = FadeInFrames * PixelsPerFrame;
 	FadeInStartX = FMath::RoundToInt32(StartTimeHandleX);
 	FadeInEndX = FMath::RoundToInt32(FMath::Clamp(StartTimeHandleX + FadeInPixelLength, StartTimeHandleX, EndTimeHandleX));
@@ -161,14 +163,17 @@ void FWaveformTransformationTrimFadeRenderer::GenerateFadeCurves(const FGeometry
 	for (uint32 Pixel = 0; Pixel < DisplayedFadeInPixelLength; ++Pixel)
 	{
 		const double FadeFraction = (float)Pixel / FadeInPixelLength;
-		const double CurveFunction = UWaveformTransformationTrimFade::GetFadeInCurveValue(StrongTrimFade->StartFadeCurve, FadeFraction, StrongTrimFade->StartSCurveSharpness);
+		const double CurveFunction = StrongTrimFade->FadeFunctions.FadeIn ? StrongTrimFade->FadeFunctions.FadeIn->GetFadeInCurveValue(FadeFraction)
+			: UWaveformTransformationTrimFade::GetFadeInCurveValue(StrongTrimFade->StartFadeCurve, FadeFraction, StrongTrimFade->StartSCurveSharpness);
+
 		const double CurveValue = Pixel != FadeInPixelLength - 1 ? 1.f - CurveFunction : 0.f;
 
 		const uint32 XCoordinate = Pixel + FadeInStartX;
 		FadeInCurvePoints[Pixel] = FVector2D(XCoordinate, CurveValue * AllottedGeometry.Size.Y);
 	}
 
-	const float FadeOutFrames = StrongTrimFade->EndFadeTime * TransformationWaveInfo.SampleRate;
+	const float FadeOutTime = StrongTrimFade->FadeFunctions.FadeOut ? StrongTrimFade->FadeFunctions.FadeOut->Duration : StrongTrimFade->EndFadeTime;
+	const float FadeOutFrames = FadeOutTime * TransformationWaveInfo.SampleRate;
 	const float FadeOutPixelLength = FadeOutFrames * PixelsPerFrame;
 	FadeOutStartX = FMath::RoundToInt32(FMath::Clamp(EndTimeHandleX - FadeOutPixelLength, StartTimeHandleX, EndTimeHandleX));
 	FadeOutEndX = FMath::RoundToInt32(EndTimeHandleX);
@@ -180,13 +185,16 @@ void FWaveformTransformationTrimFadeRenderer::GenerateFadeCurves(const FGeometry
 	for (uint32 Pixel = 0; Pixel < DisplayedFadeOutPixelLength; ++Pixel)
 	{
 		const double FadeFraction = (float)(Pixel + FadeOutPixelOffset) / FadeOutPixelLength;
-		const double CurveFunction = UWaveformTransformationTrimFade::GetFadeOutCurveValue(StrongTrimFade->EndFadeCurve, FadeFraction, StrongTrimFade->EndSCurveSharpness);
+		const double CurveFunction = StrongTrimFade->FadeFunctions.FadeOut ? StrongTrimFade->FadeFunctions.FadeOut->GetFadeOutCurveValue(FadeFraction)
+			: UWaveformTransformationTrimFade::GetFadeOutCurveValue(StrongTrimFade->EndFadeCurve, FadeFraction, StrongTrimFade->EndSCurveSharpness);
+
 		const double CurveValue = Pixel != FadeOutPixelLength - 1 ? 1.f - CurveFunction : 1.f;
 
 		const uint32 XCoordinate = Pixel + FadeOutStartX;
 		FadeOutCurvePoints[Pixel] = FVector2D(XCoordinate, CurveValue * AllottedGeometry.Size.Y);
 	}
 }
+PRAGMA_ENABLE_DEPRECATION_WARNINGS
 
 FCursorReply FWaveformTransformationTrimFadeRenderer::OnCursorQuery(const FGeometry& MyGeometry, const FPointerEvent& CursorEvent) const
 {
@@ -226,26 +234,93 @@ FReply FWaveformTransformationTrimFadeRenderer::OnMouseWheel(SWidget& OwnerWidge
 	if (IsCursorInFadeInInteractionRange(LocalCursorPosition, MyGeometry))
 	{		
 		const float FadeCurveDelta = MouseEvent.GetWheelDelta() * MouseWheelStep;
-		const float FadeCurveValue = FMath::Clamp(StrongTrimFade->StartFadeCurve + FadeCurveDelta, 0.f, 10.f);
+		FText OutPropertyName = FText::FromName(NAME_None);
+		
+		if (StrongTrimFade->FadeFunctions.FadeIn && StrongTrimFade->FadeFunctions.FadeIn.IsA<UFadeCurveFunction>())
+		{
+			TObjectPtr<UFadeCurveFunction> FadeIn = CastChecked<UFadeCurveFunction>(StrongTrimFade->FadeFunctions.FadeIn);
+			OutPropertyName = FadeIn->GetFadeCurvePropertyName();
+			
+			if (BeginTransaction(TEXT("PropertyEditor"), FText::Format(NSLOCTEXT("PropertyEditor", "SetNumericPropertyTransaction", "Edit {0}"), OutPropertyName), StrongTrimFade.Get()) == INDEX_NONE)
+			{
+				UE_LOG(LogWaveformTransformation, Warning, TEXT("Begin TrimFadeTransformation OnMouseWheel Transaction Failed"));
+			}
+
+			StrongTrimFade->Modify();
+			StrongTrimFade->FadeFunctions.FadeIn->Modify();
+
+			const float FadeCurveValue = FadeIn->GetFadeCurve() + FadeCurveDelta;
+			FadeIn->SetFadeCurve(FadeCurveValue);
+		}
+		else
+		{
+			PRAGMA_DISABLE_DEPRECATION_WARNINGS
+			OutPropertyName = FText::FromName(GET_MEMBER_NAME_CHECKED(UWaveformTransformationTrimFade, StartFadeCurve));
+			if (BeginTransaction(TEXT("PropertyEditor"), FText::Format(NSLOCTEXT("PropertyEditor", "SetNumericPropertyTransaction", "Edit {0}"), OutPropertyName), StrongTrimFade.Get()) == INDEX_NONE)
+			{
+				UE_LOG(LogWaveformTransformation, Warning, TEXT("Begin TrimFadeTransformation OnMouseWheel Transaction Failed"));
+			}
+
+			StrongTrimFade->Modify();
+
+			const float FadeCurveValue = FMath::Clamp(StrongTrimFade->StartFadeCurve + FadeCurveDelta, 0.f, 10.f);
+			StrongTrimFade->StartFadeCurve = FadeCurveValue;
+			PRAGMA_ENABLE_DEPRECATION_WARNINGS
+		}
 
 		EPropertyValueSetFlags::Type PropertyFlags = EPropertyValueSetFlags::InteractiveChange;
-		StrongTrimFade->StartFadeCurve = FadeCurveValue;
 
 		StrongTrimFade->PostEditChange();
 
+		EndTransaction();
+		StrongTrimFade->OnTransformationChanged.ExecuteIfBound();
+
 		return FReply::Handled();
 	}
 
 	if (IsCursorInFadeOutInteractionRange(LocalCursorPosition, MyGeometry))
 	{
 		const float FadeCurveDelta = MouseEvent.GetWheelDelta() * MouseWheelStep;
-		const float FadeCurveValue = FMath::Clamp(StrongTrimFade->EndFadeCurve + FadeCurveDelta, 0.f, 10.f);
+		FText OutPropertyName = FText::FromName(NAME_None);
+
+		if (StrongTrimFade->FadeFunctions.FadeOut && StrongTrimFade->FadeFunctions.FadeOut.IsA<UFadeCurveFunction>())
+		{
+			TObjectPtr<UFadeCurveFunction> FadeOut = CastChecked<UFadeCurveFunction>(StrongTrimFade->FadeFunctions.FadeOut);
+			OutPropertyName = FadeOut->GetFadeCurvePropertyName();
+			if (BeginTransaction(TEXT("PropertyEditor"), FText::Format(NSLOCTEXT("PropertyEditor", "SetNumericPropertyTransaction", "Edit {0}"), OutPropertyName), StrongTrimFade.Get()) == INDEX_NONE)
+			{
+				UE_LOG(LogWaveformTransformation, Warning, TEXT("Begin TrimFadeTransformation OnMouseWheel Transaction Failed"));
+			}
+
+			StrongTrimFade->Modify();
+			StrongTrimFade->FadeFunctions.FadeOut->Modify();
+
+			const float FadeCurveValue = FadeOut->GetFadeCurve() + FadeCurveDelta;
+			FadeOut->SetFadeCurve(FadeCurveValue);
+		}
+		else
+		{
+			PRAGMA_DISABLE_DEPRECATION_WARNINGS
+			OutPropertyName = FText::FromName(GET_MEMBER_NAME_CHECKED(UWaveformTransformationTrimFade, EndFadeCurve));
+			if (BeginTransaction(TEXT("PropertyEditor"), FText::Format(NSLOCTEXT("PropertyEditor", "SetNumericPropertyTransaction", "Edit {0}"), OutPropertyName), StrongTrimFade.Get()) == INDEX_NONE)
+			{
+				UE_LOG(LogWaveformTransformation, Warning, TEXT("Begin TrimFadeTransformation OnMouseWheel Transaction Failed"));
+			}
+
+			StrongTrimFade->Modify();
+
+			const float FadeCurveValue = FMath::Clamp(StrongTrimFade->EndFadeCurve + FadeCurveDelta, 0.f, 10.f);
+			StrongTrimFade->EndFadeCurve = FadeCurveValue;
+			PRAGMA_ENABLE_DEPRECATION_WARNINGS
+		}
 		
 		EPropertyValueSetFlags::Type PropertyFlags = EPropertyValueSetFlags::InteractiveChange;
-		StrongTrimFade->EndFadeCurve = FadeCurveValue;
 
 		StrongTrimFade->PostEditChange();
 		
+		EndTransaction();
+		StrongTrimFade->OnTransformationChanged.ExecuteIfBound();
+
 		return FReply::Handled();
 	}
 
@@ -380,23 +455,47 @@ void FWaveformTransformationTrimFadeRenderer::SetPropertyValueDependingOnInterac
 		break;
 	case FWaveformTransformationTrimFadeRenderer::ETrimFadeInteractionType::ScrubbingFadeIn:
 		{
-			float StartFadeTimeValue = FMath::Clamp(SelectedTime - StrongTrimFade->StartTime, 0.f, TNumericLimits<float>().Max());
-			StrongTrimFade->StartFadeTime = StartFadeTimeValue;
+			float StartFadeTimeValue = 0;
+
+			if (StrongTrimFade->FadeFunctions.FadeIn)
+			{
+				StrongTrimFade->FadeFunctions.FadeIn->Modify();
+				StartFadeTimeValue = FMath::Clamp(SelectedTime - StrongTrimFade->StartTime, 0.f, TNumericLimits<float>().Max());
+				StrongTrimFade->FadeFunctions.FadeIn->Duration = StartFadeTimeValue;
+			}
+			else
+			{
+				PRAGMA_DISABLE_DEPRECATION_WARNINGS
+				StartFadeTimeValue = FMath::Clamp(SelectedTime - StrongTrimFade->StartTime, 0.f, TNumericLimits<float>().Max());
+				StrongTrimFade->StartFadeTime = StartFadeTimeValue;
+				PRAGMA_ENABLE_DEPRECATION_WARNINGS
+			}
 		}
 		break;
 	case FWaveformTransformationTrimFadeRenderer::ETrimFadeInteractionType::ScrubbingFadeOut:
 		{	
-			float EndFadeTimeValue = FMath::Clamp(StrongTrimFade->EndTime - SelectedTime, 0.f, TNumericLimits<float>().Max());
-			StrongTrimFade->EndFadeTime = EndFadeTimeValue;
+			float EndFadeTimeValue = 0;
+
+			if (StrongTrimFade->FadeFunctions.FadeOut)
+			{
+				StrongTrimFade->FadeFunctions.FadeOut->Modify();
+				EndFadeTimeValue = FMath::Clamp(StrongTrimFade->EndTime - SelectedTime, 0.f, TNumericLimits<float>().Max());
+				StrongTrimFade->FadeFunctions.FadeOut->Duration = EndFadeTimeValue;
+			}
+			else
+			{
+				PRAGMA_DISABLE_DEPRECATION_WARNINGS
+				EndFadeTimeValue = FMath::Clamp(StrongTrimFade->EndTime - SelectedTime, 0.f, TNumericLimits<float>().Max());
+				StrongTrimFade->EndFadeTime = EndFadeTimeValue;
+				PRAGMA_ENABLE_DEPRECATION_WARNINGS
+			}
 		}
 		break;
 	case FWaveformTransformationTrimFadeRenderer::ETrimFadeInteractionType::RightClickFadeIn:
-		CurvePropertyToChange = MakeShared<float>(StrongTrimFade->StartFadeCurve);
-		ShowSelectFadeModeMenuAtCursor(WidgetGeometry, MouseEvent);
-		break;
 	case FWaveformTransformationTrimFadeRenderer::ETrimFadeInteractionType::RightClickFadeOut:
-		CurvePropertyToChange = MakeShared<float>(StrongTrimFade->EndFadeCurve);
+	{
 		ShowSelectFadeModeMenuAtCursor(WidgetGeometry, MouseEvent);
+	}
 		break;
 	default:
 		break;
@@ -449,14 +548,24 @@ void FWaveformTransformationTrimFadeRenderer::ShowSelectFadeModeMenuAtCursor(con
 
 							PinnedTrimFade->Modify();
 
+							PRAGMA_DISABLE_DEPRECATION_WARNINGS
 							if (InteractionType == FWaveformTransformationTrimFadeRenderer::ETrimFadeInteractionType::RightClickFadeIn)
 							{
+								if (PinnedTrimFade->FadeFunctions.FadeIn)
+								{
+									PinnedTrimFade->FadeFunctions.FadeIn = NewObject<UFadeFunction>(PinnedTrimFade.Get(), UWaveformTransformationTrimFade::FadeModeToFadeFunctionMap[FadeOptionPair.Key].Get());
+								}
 								PinnedTrimFade->StartFadeCurve = FadeOptionPair.Value;
 							}
 							else if (InteractionType == FWaveformTransformationTrimFadeRenderer::ETrimFadeInteractionType::RightClickFadeOut)
 							{
+								if (PinnedTrimFade->FadeFunctions.FadeOut)
+								{
+									PinnedTrimFade->FadeFunctions.FadeOut = NewObject<UFadeFunction>(PinnedTrimFade.Get(), UWaveformTransformationTrimFade::FadeModeToFadeFunctionMap[FadeOptionPair.Key].Get());
+								}
 								PinnedTrimFade->EndFadeCurve = FadeOptionPair.Value;
 							}
+							PRAGMA_ENABLE_DEPRECATION_WARNINGS
 
 							if (GEditor && GEditor->Trans)
 							{
@@ -555,7 +664,7 @@ const FSlateIcon FWaveformTransformationTrimFadeRenderer::GetFadeModeIcon(const
 		{
 		case EWaveEditorFadeMode::Linear:
 			return FSlateIcon(FAudioWidgetsStyle::StyleName, "AudioWidgetsStyle.FadeInLinear");
-		case EWaveEditorFadeMode::Exponetial:
+		case EWaveEditorFadeMode::Exponential:
 			return FSlateIcon(FAudioWidgetsStyle::StyleName, "AudioWidgetsStyle.FadeInExponential");
 		case EWaveEditorFadeMode::Logarithmic:
 			return FSlateIcon(FAudioWidgetsStyle::StyleName, "AudioWidgetsStyle.FadeInLogarithmic");
@@ -571,7 +680,7 @@ const FSlateIcon FWaveformTransformationTrimFadeRenderer::GetFadeModeIcon(const
 		{
 		case EWaveEditorFadeMode::Linear:
 			return FSlateIcon(FAudioWidgetsStyle::StyleName, "AudioWidgetsStyle.FadeOutLinear");
-		case EWaveEditorFadeMode::Exponetial:
+		case EWaveEditorFadeMode::Exponential:
 			return FSlateIcon(FAudioWidgetsStyle::StyleName, "AudioWidgetsStyle.FadeOutExponential");
 		case EWaveEditorFadeMode::Logarithmic:
 			return FSlateIcon(FAudioWidgetsStyle::StyleName, "AudioWidgetsStyle.FadeOutLogarithmic");
@@ -652,10 +761,32 @@ FText FWaveformTransformationTrimFadeRenderer::GetPropertyEditedByCurrentInterac
 		OutPropertyName = FText::FromName(GET_MEMBER_NAME_CHECKED(UWaveformTransformationTrimFade, EndTime));
 		break;
 	case FWaveformTransformationTrimFadeRenderer::ETrimFadeInteractionType::ScrubbingFadeIn:
-		OutPropertyName = FText::FromName(GET_MEMBER_NAME_CHECKED(UWaveformTransformationTrimFade, StartFadeTime));
+		if (StrongTrimFade->FadeFunctions.FadeIn)
+		{
+			OutPropertyName = FText::Format(FText::FromString("{0}::{1}"),
+				FText::FromName(GET_MEMBER_NAME_CHECKED(FFadeFunctionData, FadeIn)),
+				FText::FromName(GET_MEMBER_NAME_CHECKED(UFadeFunction, Duration)));
+		}
+		else
+		{
+			PRAGMA_DISABLE_DEPRECATION_WARNINGS
+			OutPropertyName = FText::FromName(GET_MEMBER_NAME_CHECKED(UWaveformTransformationTrimFade, StartFadeTime));
+			PRAGMA_ENABLE_DEPRECATION_WARNINGS
+		}
 		break;
 	case FWaveformTransformationTrimFadeRenderer::ETrimFadeInteractionType::ScrubbingFadeOut:
-		OutPropertyName = FText::FromName(GET_MEMBER_NAME_CHECKED(UWaveformTransformationTrimFade, EndFadeTime));
+		if (StrongTrimFade->FadeFunctions.FadeOut)
+		{
+			OutPropertyName = FText::Format(FText::FromString("{0}::{1}"),
+				FText::FromName(GET_MEMBER_NAME_CHECKED(FFadeFunctionData, FadeOut)),
+				FText::FromName(GET_MEMBER_NAME_CHECKED(UFadeFunction, Duration)));
+		}
+		else
+		{
+			PRAGMA_DISABLE_DEPRECATION_WARNINGS
+			OutPropertyName = FText::FromName(GET_MEMBER_NAME_CHECKED(UWaveformTransformationTrimFade, EndFadeTime));
+			PRAGMA_ENABLE_DEPRECATION_WARNINGS
+		}
 		break;
 	default:
 		break;
diff --git a/Engine/Plugins/Editor/WaveformEditor/Source/WaveformTransformationsWidgets/Private/WaveformTransformationTrimFadeRenderer.h b/Engine/Plugins/Editor/WaveformEditor/Source/WaveformTransformationsWidgets/Private/WaveformTransformationTrimFadeRenderer.h
--- a/Engine/Plugins/Editor/WaveformEditor/Source/WaveformTransformationsWidgets/Private/WaveformTransformationTrimFadeRenderer.h
+++ b/Engine/Plugins/Editor/WaveformEditor/Source/WaveformTransformationsWidgets/Private/WaveformTransformationTrimFadeRenderer.h
@@ -70,7 +70,6 @@ private:
 	const FSlateIcon GetFadeModeIcon(const EWaveEditorFadeMode& FadeMode);
 
 	TWeakPtr<SWindow> FadeModeMenuWindow = nullptr;
-	TSharedPtr<float> CurvePropertyToChange = nullptr;
 
 	TStrongObjectPtr<UWaveformTransformationTrimFade> StrongTrimFade = nullptr;
 
diff --git a/Engine/Plugins/Editor/WaveformEditor/Source/WaveformTransformationsWidgets/Public/WaveformTransformationRendererBase.h b/Engine/Plugins/Editor/WaveformEditor/Source/WaveformTransformationsWidgets/Public/WaveformTransformationRendererBase.h
--- a/Engine/Plugins/Editor/WaveformEditor/Source/WaveformTransformationsWidgets/Public/WaveformTransformationRendererBase.h
+++ b/Engine/Plugins/Editor/WaveformEditor/Source/WaveformTransformationsWidgets/Public/WaveformTransformationRendererBase.h
@@ -60,7 +60,7 @@ protected:
 
 	static constexpr float InteractionPixelXDelta = 10;
 	static constexpr float InteractionRatioYDelta = 0.07f;
-	static constexpr float MouseWheelStep = 0.25f;
+	static constexpr float MouseWheelStep = 0.1f;
 
 private:
 	UE_DEPRECATED(5.7, "PropertyHandles have been deprecated.")
