[Waveform] Editor Crashes when User attempts to change the Waveform Transformations Display colors
#UE-290368

diff --git a/Engine/Plugins/Editor/WaveformEditor/Source/WaveformTransformationsWidgets/Private/WaveformTransformationMarkersRenderer.cpp b/Engine/Plugins/Editor/WaveformEditor/Source/WaveformTransformationsWidgets/Private/WaveformTransformationMarkersRenderer.cpp
--- a/Engine/Plugins/Editor/WaveformEditor/Source/WaveformTransformationsWidgets/Private/WaveformTransformationMarkersRenderer.cpp
+++ b/Engine/Plugins/Editor/WaveformEditor/Source/WaveformTransformationsWidgets/Private/WaveformTransformationMarkersRenderer.cpp
@@ -19,7 +19,28 @@ FWaveformTransformationMarkerRenderer::FWaveformTransformationMarkerRenderer()
 {
 	const UWaveformTransformationsWidgetsSettings* Settings = GetDefault<UWaveformTransformationsWidgetsSettings>();
 	check(Settings);
-	Settings->OnSettingChanged().AddRaw(this, &FWaveformTransformationMarkerRenderer::OnWidgetSettingsUpdated);
+	OnSettingChangedDelegateHandle = Settings->OnSettingChanged().AddLambda([this](const FName& PropertyName, const UWaveformTransformationsWidgetsSettings* Settings)
+		{
+			check(Settings)
+
+			if (PropertyName == GET_MEMBER_NAME_CHECKED(UWaveformTransformationsWidgetsSettings, MarkerColor))
+			{
+				MarkerColor = Settings->MarkerColor;
+			}
+			if (PropertyName == GET_MEMBER_NAME_CHECKED(UWaveformTransformationsWidgetsSettings, LoopColors))
+			{
+				SetLoopColors(Settings->LoopColors);
+			}
+			if (PropertyName == GET_MEMBER_NAME_CHECKED(UWaveformTransformationsWidgetsSettings, LabelTextColor))
+			{
+				LabelTextColor = Settings->LabelTextColor;
+			}
+			if (PropertyName == GET_MEMBER_NAME_CHECKED(UWaveformTransformationsWidgetsSettings, LabelFontSize))
+			{
+				LabelFontSize = Settings->LabelFontSize;
+			}
+		});
+
 	ApplyWidgetSettings(Settings);
 }
 
@@ -27,6 +48,13 @@ FWaveformTransformationMarkerRenderer::~FWaveformTransformationMarkerRenderer()
 {
 	FSlateApplication::Get().OnFocusChanging().Remove(PopupHandle);
 	FSlateApplication::Get().OnApplicationActivationStateChanged().Remove(ApplicationActivationStateHandle);
+
+	if (OnSettingChangedDelegateHandle.IsValid())
+	{
+		const UWaveformTransformationsWidgetsSettings* Settings = GetDefault<UWaveformTransformationsWidgetsSettings>();
+		check(Settings);
+		Settings->OnSettingChanged().Remove(OnSettingChangedDelegateHandle);
+	}
 }
 
 int32 FWaveformTransformationMarkerRenderer::OnPaint(const FPaintArgs& Args, const FGeometry& AllottedGeometry, const FSlateRect& MyCullingRect, FSlateWindowElementList& OutDrawElements, int32 LayerId, const FWidgetStyle& InWidgetStyle, bool bParentEnabled) const
@@ -463,28 +491,6 @@ FReply FWaveformTransformationMarkerRenderer::OnKeyUp(const FGeometry& MyGeometr
 	return FReply::Handled();
 }
 
-void FWaveformTransformationMarkerRenderer::OnWidgetSettingsUpdated(const FName& PropertyName, const UWaveformTransformationsWidgetsSettings* Settings)
-{
-	check(Settings)
-
-	if (PropertyName == GET_MEMBER_NAME_CHECKED(UWaveformTransformationsWidgetsSettings, MarkerColor))
-	{
-		MarkerColor = Settings->MarkerColor;
-	}
-	if (PropertyName == GET_MEMBER_NAME_CHECKED(UWaveformTransformationsWidgetsSettings, LoopColors))
-	{
-		SetLoopColors(Settings->LoopColors);
-	}
-	if (PropertyName == GET_MEMBER_NAME_CHECKED(UWaveformTransformationsWidgetsSettings, LabelTextColor))
-	{
-		LabelTextColor = Settings->LabelTextColor;
-	}
-	if (PropertyName == GET_MEMBER_NAME_CHECKED(UWaveformTransformationsWidgetsSettings, LabelFontSize))
-	{
-		LabelFontSize = Settings->LabelFontSize;
-	}
-}
-
 void FWaveformTransformationMarkerRenderer::ApplyWidgetSettings(const UWaveformTransformationsWidgetsSettings* Settings)
 {
 	check(Settings)
diff --git a/Engine/Plugins/Editor/WaveformEditor/Source/WaveformTransformationsWidgets/Private/WaveformTransformationMarkersRenderer.h b/Engine/Plugins/Editor/WaveformEditor/Source/WaveformTransformationsWidgets/Private/WaveformTransformationMarkersRenderer.h
--- a/Engine/Plugins/Editor/WaveformEditor/Source/WaveformTransformationsWidgets/Private/WaveformTransformationMarkersRenderer.h
+++ b/Engine/Plugins/Editor/WaveformEditor/Source/WaveformTransformationsWidgets/Private/WaveformTransformationMarkersRenderer.h
@@ -27,7 +27,6 @@ public:
 private:
 	int32 DrawMarkerHandles(const FGeometry& AllottedGeometry, FSlateWindowElementList& OutDrawElements, int32 LayerId) const;
 
-	void OnWidgetSettingsUpdated(const FName& PropertyName, const UWaveformTransformationsWidgetsSettings* Settings);
 	void ApplyWidgetSettings(const UWaveformTransformationsWidgetsSettings* Settings);
 
 	void SetLoopColors(const TArray<FLinearColor>& InColors);
@@ -100,6 +99,7 @@ private:
 
 	FDelegateHandle PopupHandle;
 	FDelegateHandle ApplicationActivationStateHandle;
+	FDelegateHandle OnSettingChangedDelegateHandle;
 
 	TWeakPtr<SWindow> CreateLoopMenuWindow;
 };
