Waveform Editor: Fix selecting a loop region marks the soundwave dirty

diff --git a/Engine/Plugins/Editor/WaveformEditor/Source/WaveformEditor/Private/WaveformEditor.cpp b/Engine/Plugins/Editor/WaveformEditor/Source/WaveformEditor/Private/WaveformEditor.cpp
--- a/Engine/Plugins/Editor/WaveformEditor/Source/WaveformEditor/Private/WaveformEditor.cpp
+++ b/Engine/Plugins/Editor/WaveformEditor/Source/WaveformEditor/Private/WaveformEditor.cpp
@@ -139,11 +139,12 @@ bool FWaveformEditor::Init(const EToolkitMode::Type Mode, const TSharedPtr<ITool
 
 			if (!MarkerTransformation->OnTransformationChanged.IsBound())
 			{
-				MarkerTransformation->OnTransformationChanged.BindLambda([WeakWaveformEditor]()
+				MarkerTransformation->OnTransformationChanged.BindLambda([WeakWaveformEditor](bool bModify)
 					{
 						if (TSharedPtr<FWaveformEditor> PinnedEditor = WeakWaveformEditor.Pin())
 						{
-							PinnedEditor->UpdateTransformations();
+							PinnedEditor->UpdateTransformations(bModify);
+
 						}
 					});
 			}
@@ -168,11 +169,11 @@ bool FWaveformEditor::Init(const EToolkitMode::Type Mode, const TSharedPtr<ITool
 		{
 			if (!TrimFadeTransformation->OnTransformationChanged.IsBound())
 			{
-				TrimFadeTransformation->OnTransformationChanged.BindLambda([WeakWaveformEditor]()
+				TrimFadeTransformation->OnTransformationChanged.BindLambda([WeakWaveformEditor](bool bModify)
 					{
 						if (TSharedPtr<FWaveformEditor> PinnedEditor = WeakWaveformEditor.Pin())
 						{
-							PinnedEditor->UpdateTransformations();
+							PinnedEditor->UpdateTransformations(bModify);
 						}
 					});
 			}
@@ -762,12 +763,35 @@ void FWaveformEditor::RegenerateTransformations()
 	UpdateTransportState();
 }
 
-void FWaveformEditor::UpdateTransformations()
+void FWaveformEditor::UpdateTransformations(bool bMarkDirty)
 {
 	check(TransportController);
+	check(SoundWave);
 
 	bool bIsPlaying = TransportController->IsPlaying();
-	NotifyPostTransformationChange(EPropertyChangeType::ValueSet);
+
+	FProperty* TransformationsProperty = FindFProperty<FProperty>(USoundWave::StaticClass(), GET_MEMBER_NAME_CHECKED(USoundWave, Transformations));
+	check(TransformationsProperty);
+
+	FEditPropertyChain PropertyChain;
+	PropertyChain.AddHead(TransformationsProperty);
+	PropertyChain.SetActivePropertyNode(TransformationsProperty);
+
+	FPropertyChangedEvent PropertyChangedEvent(TransformationsProperty, EPropertyChangeType::ValueSet);
+
+	NotifyPostChange(PropertyChangedEvent, &PropertyChain);
+
+	bool bWasAlreadyDirty = SoundWave->GetPackage()->IsDirty();
+
+	//Update the content browser asset state to dirty
+	SoundWave->PostEditChange();
+	TransformationChainConfig = SoundWave->UpdateTransformations();
+
+	// Check if we do not want to mark the package dirty for a user facing change they would not expect to modify data such as selecting a loop region
+	if (!bMarkDirty && !bWasAlreadyDirty)
+	{
+		SoundWave->GetPackage()->ClearDirtyFlag();
+	}
 
 	if (bIsPlaying)
 	{
@@ -1756,11 +1780,11 @@ void FWaveformEditor::NotifyPostChange(const FPropertyChangedEvent& PropertyChan
 
 		if (!MarkerTransformation->OnTransformationChanged.IsBound())
 		{
-			MarkerTransformation->OnTransformationChanged.BindLambda([WeakWaveformEditor]()
+			MarkerTransformation->OnTransformationChanged.BindLambda([WeakWaveformEditor](bool bModify)
 				{
 					if (TSharedPtr<FWaveformEditor> PinnedEditor = WeakWaveformEditor.Pin())
 					{
-						PinnedEditor->UpdateTransformations();
+						PinnedEditor->UpdateTransformations(bModify);
 					}
 				});
 		}
@@ -1787,11 +1811,11 @@ void FWaveformEditor::NotifyPostChange(const FPropertyChangedEvent& PropertyChan
 	{
 		if (!TrimFadeTransformation->OnTransformationChanged.IsBound())
 		{
-			TrimFadeTransformation->OnTransformationChanged.BindLambda([WeakWaveformEditor]()
+			TrimFadeTransformation->OnTransformationChanged.BindLambda([WeakWaveformEditor](bool bModify)
 				{
 					if (TSharedPtr<FWaveformEditor> PinnedEditor = WeakWaveformEditor.Pin())
 					{
-						PinnedEditor->UpdateTransformations();
+						PinnedEditor->UpdateTransformations(bModify);
 					}
 				});
 		}
diff --git a/Engine/Plugins/Editor/WaveformEditor/Source/WaveformEditor/Public/WaveformEditor.h b/Engine/Plugins/Editor/WaveformEditor/Source/WaveformEditor/Public/WaveformEditor.h
--- a/Engine/Plugins/Editor/WaveformEditor/Source/WaveformEditor/Public/WaveformEditor.h
+++ b/Engine/Plugins/Editor/WaveformEditor/Source/WaveformEditor/Public/WaveformEditor.h
@@ -127,7 +127,7 @@ private:
 	UE_API void DeleteMarker();
 	UE_API void SkipToNextMarker(); // Skips the playhead to the next marker
 	UE_API void RegenerateTransformations();
-	UE_API void UpdateTransformations();
+	UE_API void UpdateTransformations(bool bMarkDirty = true);
 	UE_API void UpdateRenderElements();
 	UE_API void UpdateTransportState();
 
diff --git a/Engine/Plugins/Editor/WaveformEditor/Source/WaveformTransformations/Private/WaveformTransformationMarkers.cpp b/Engine/Plugins/Editor/WaveformEditor/Source/WaveformTransformations/Private/WaveformTransformationMarkers.cpp
--- a/Engine/Plugins/Editor/WaveformEditor/Source/WaveformTransformations/Private/WaveformTransformationMarkers.cpp
+++ b/Engine/Plugins/Editor/WaveformEditor/Source/WaveformTransformations/Private/WaveformTransformationMarkers.cpp
@@ -96,7 +96,7 @@ void UWaveformTransformationMarkers::PostEditChangeProperty(FPropertyChangedEven
 	}
 	else 
 	{
-		OnTransformationChanged.ExecuteIfBound();
+		OnTransformationChanged.ExecuteIfBound(true);
 	}
 }
 
@@ -311,7 +311,7 @@ void UWaveformTransformationMarkers::ModifyMarkerLoopRegion(ELoopModificationCon
 		GEditor->EndTransaction();
 	}
 
-	OnTransformationChanged.ExecuteIfBound();
+	OnTransformationChanged.ExecuteIfBound(true);
 }
 
 void UWaveformTransformationMarkers::CycleMarkerLoopRegion(ELoopModificationControls Modification)
@@ -375,7 +375,7 @@ void UWaveformTransformationMarkers::CycleMarkerLoopRegion(ELoopModificationCont
 		ResetLoopPreviewing();
 	}
 
-	OnTransformationChanged.ExecuteIfBound();
+	OnTransformationChanged.ExecuteIfBound(false);
 }
 
 int64 UWaveformTransformationMarkers::GetFramesToNextZeroCrossing(int64 StartFrame, int64 EndFrame, const TArray<float>& PCMArray, bool bCheckIncrementing) const
@@ -527,7 +527,7 @@ void UWaveformTransformationMarkers::SetLoopPreviewing()
 
 	if (bTransformationChanged)
 	{
-		OnTransformationChanged.ExecuteIfBound();
+		OnTransformationChanged.ExecuteIfBound(true);
 	}
 }
 
diff --git a/Engine/Plugins/Editor/WaveformEditor/Source/WaveformTransformationsWidgets/Private/WaveformTransformationMarkersRenderer.cpp b/Engine/Plugins/Editor/WaveformEditor/Source/WaveformTransformationsWidgets/Private/WaveformTransformationMarkersRenderer.cpp
--- a/Engine/Plugins/Editor/WaveformEditor/Source/WaveformTransformationsWidgets/Private/WaveformTransformationMarkersRenderer.cpp
+++ b/Engine/Plugins/Editor/WaveformEditor/Source/WaveformTransformationsWidgets/Private/WaveformTransformationMarkersRenderer.cpp
@@ -302,12 +302,12 @@ void FWaveformTransformationMarkerRenderer::Tick(const FGeometry& AllottedGeomet
 		if (Marker.CuePointID == MarkersArray->SelectedCue && Marker.IsLoopRegion() &&
 			(Marker.FramePosition + Marker.FrameLength < StartFrameOffset || Marker.FramePosition > StartFrameOffset + NumAvailableSamples))
 		{
-			MarkersArray->Modify();
+			MarkersArray->Modify(bMarkFileDirty);
 			MarkersArray->SelectedCue = INDEX_NONE;
 			
 			if (StrongMarkersTransformation->ResetLoopPreviewing())
 			{
-				StrongMarkersTransformation->OnTransformationChanged.ExecuteIfBound();
+				StrongMarkersTransformation->OnTransformationChanged.ExecuteIfBound(true);
 			}
 
 			//End move transaction
@@ -317,7 +317,7 @@ void FWaveformTransformationMarkerRenderer::Tick(const FGeometry& AllottedGeomet
 
 	if (StrongMarkersTransformation->AdjustLoopPreviewIfNotAligned())
 	{
-		StrongMarkersTransformation->OnTransformationChanged.ExecuteIfBound();
+		StrongMarkersTransformation->OnTransformationChanged.ExecuteIfBound(true);
 	}
 
 	//Invalidate Marker pointer if it has been removed from the array
@@ -359,8 +359,6 @@ FReply FWaveformTransformationMarkerRenderer::OnMouseButtonDown(SWidget& OwnerWi
 
 	if (MarkerInteractionType != EMarkerInteractionType::None)
 	{
-		MarkersArray->Modify();
-
 		return FReply::Handled().CaptureMouse(OwnerWidget.AsShared()).PreventThrottling();
 	}
 
@@ -389,9 +387,11 @@ FReply FWaveformTransformationMarkerRenderer::OnMouseButtonUp(SWidget& OwnerWidg
 
 	if (MarkerInteractionType != EMarkerInteractionType::None)
 	{
-		StrongMarkersTransformation->Modify();
+		bool bMarkDirty = MarkerInteractionType == EMarkerInteractionType::ScrubbingMarkerHandleLeft || MarkerInteractionType == EMarkerInteractionType::ScrubbingMarkerHandleRight;
+
+		StrongMarkersTransformation->Modify(bMarkDirty);
 		SetPropertyValueDependingOnInteractionType(MouseEvent, MyGeometry);
-		
+
 		// If we know it was an interaction that started a transaction, end the transaction
 		if (MarkerInteractionType == EMarkerInteractionType::ScrubbingMarkerHandleLeft || MarkerInteractionType == EMarkerInteractionType::ScrubbingMarkerHandleRight
 			|| MarkerInteractionType == EMarkerInteractionType::LoopHandle || MarkerInteractionType == EMarkerInteractionType::MarkerHandle)
@@ -401,14 +401,14 @@ FReply FWaveformTransformationMarkerRenderer::OnMouseButtonUp(SWidget& OwnerWidg
 
 		MarkerInteractionType = EMarkerInteractionType::None;
 		
-		StrongMarkersTransformation->OnTransformationChanged.ExecuteIfBound();
+		StrongMarkersTransformation->OnTransformationChanged.ExecuteIfBound(false);
 		
 		return FReply::Handled().ReleaseMouseCapture();
 	}
 
 	if (bHasTransformationChanged)
 	{
-		StrongMarkersTransformation->OnTransformationChanged.ExecuteIfBound();
+		StrongMarkersTransformation->OnTransformationChanged.ExecuteIfBound(true);
 	}
 
 	return FReply::Unhandled();
@@ -614,12 +614,23 @@ bool FWaveformTransformationMarkerRenderer::SetMarkerInInteractionRange(const FV
 			MarkerInInteractionRange.SetMarkerInInteractionRange(&CuePoint, false, LoopRegionPosition);
 			MarkersArray->SelectedCue = CuePoint.CuePointID;
 
+			SetActiveLoopRegion(MarkersArray->SelectedCue);
+
 			return true;
 		}
 	}
 
 	// MarkerPtr is set to null if there are no markers in range to handle cases where the mouse is not over any markers
 	MarkerInInteractionRange.SetMarkerInInteractionRange(nullptr);
+
+	if (MarkersArray->SelectedCue == INDEX_NONE)
+	{
+		StrongMarkersTransformation->ResetLoopPreviewing();
+	}
+	else
+	{
+		PreviewSelectedLoop();
+	}
 	
 	return false;
 }
@@ -650,7 +661,7 @@ void FWaveformTransformationMarkerRenderer::SetActiveLoopRegion(int32 CuePointID
 	check(StrongMarkersTransformation.IsValid());
 	check(MarkersArray != nullptr);
 
-	MarkersArray->Modify();
+	MarkersArray->Modify(false);
 	MarkersArray->SelectedCue = CuePointID;
 
 	check(TransformationWaveInfo.SampleRate > 0);
@@ -746,7 +757,7 @@ FWaveformTransformationMarkerRenderer::EMarkerInteractionType FWaveformTransform
 		{
 			SetMarkerInInteractionRange(InLocalCursorPosition, WidgetGeometry);
 			BeginTransaction(TEXT("PropertyEditor"), LOCTEXT("ResizingMarker", "Resizing Marker"), StrongMarkersTransformation.Get());
-			MarkersArray->Modify();
+			MarkersArray->Modify(bMarkFileDirty);
 			return EMarkerInteractionType::ScrubbingMarkerHandleLeft;
 		}
 
@@ -754,7 +765,7 @@ FWaveformTransformationMarkerRenderer::EMarkerInteractionType FWaveformTransform
 		{
 			SetMarkerInInteractionRange(InLocalCursorPosition, WidgetGeometry);
 			BeginTransaction(TEXT("PropertyEditor"), LOCTEXT("ResizingMarker", "Resizing Marker"), StrongMarkersTransformation.Get());
-			MarkersArray->Modify();
+			MarkersArray->Modify(bMarkFileDirty);
 			return EMarkerInteractionType::ScrubbingMarkerHandleRight;
 		}
 	}
@@ -762,7 +773,7 @@ FWaveformTransformationMarkerRenderer::EMarkerInteractionType FWaveformTransform
 	if (SetMarkerInInteractionRange(InLocalCursorPosition, WidgetGeometry))
 	{
 		BeginTransaction(TEXT("PropertyEditor"), LOCTEXT("MoveMarker", "Move Marker"), StrongMarkersTransformation.Get());
-		MarkersArray->Modify();
+		MarkersArray->Modify(bMarkFileDirty);
 		return EMarkerInteractionType::MarkerHandle;
 	}
 	
@@ -770,7 +781,7 @@ FWaveformTransformationMarkerRenderer::EMarkerInteractionType FWaveformTransform
 	{
 		SetMarkerInInteractionRange(InLocalCursorPosition, WidgetGeometry);
 		BeginTransaction(TEXT("PropertyEditor"), LOCTEXT("MoveLoop", "Move Loop"), StrongMarkersTransformation.Get());
-		MarkersArray->Modify();
+		MarkersArray->Modify(bMarkFileDirty);
 		return EMarkerInteractionType::LoopHandle;
 	}
 	else
@@ -885,7 +896,7 @@ void FWaveformTransformationMarkerRenderer::ShowConvertToLoopAtMarker(const FGeo
 								GEditor->BeginTransaction(TEXT("PropertyEditor"), LOCTEXT("ConvertMarkerToLoop", "Convert Marker To Loop"), LockedMarkers.Get());
 							}
 
-							LockedMarkers->Markers->Modify();
+							LockedMarkers->Markers->Modify(bMarkFileDirty);
 
 							if (CueToModify->FrameLength < MinLoopSize)
 							{
diff --git a/Engine/Plugins/Editor/WaveformEditor/Source/WaveformTransformationsWidgets/Private/WaveformTransformationMarkersRenderer.h b/Engine/Plugins/Editor/WaveformEditor/Source/WaveformTransformationsWidgets/Private/WaveformTransformationMarkersRenderer.h
--- a/Engine/Plugins/Editor/WaveformEditor/Source/WaveformTransformationsWidgets/Private/WaveformTransformationMarkersRenderer.h
+++ b/Engine/Plugins/Editor/WaveformEditor/Source/WaveformTransformationsWidgets/Private/WaveformTransformationMarkersRenderer.h
@@ -96,6 +96,7 @@ private:
 	float LabelFontSize = WaveformTransformationWidgetSharedDefaults::DefaultLabelFontSize;
 
 	TStrongObjectPtr<UWaveformTransformationMarkers> StrongMarkersTransformation = nullptr;
+	static const bool bMarkFileDirty = false;
 
 	FDelegateHandle PopupHandle;
 	FDelegateHandle ApplicationActivationStateHandle;
diff --git a/Engine/Plugins/Editor/WaveformEditor/Source/WaveformTransformationsWidgets/Private/WaveformTransformationTrimFadeRenderer.cpp b/Engine/Plugins/Editor/WaveformEditor/Source/WaveformTransformationsWidgets/Private/WaveformTransformationTrimFadeRenderer.cpp
--- a/Engine/Plugins/Editor/WaveformEditor/Source/WaveformTransformationsWidgets/Private/WaveformTransformationTrimFadeRenderer.cpp
+++ b/Engine/Plugins/Editor/WaveformEditor/Source/WaveformTransformationsWidgets/Private/WaveformTransformationTrimFadeRenderer.cpp
@@ -273,7 +273,7 @@ FReply FWaveformTransformationTrimFadeRenderer::OnMouseWheel(SWidget& OwnerWidge
 		StrongTrimFade->PostEditChange();
 
 		EndTransaction();
-		StrongTrimFade->OnTransformationChanged.ExecuteIfBound();
+		StrongTrimFade->OnTransformationChanged.ExecuteIfBound(true);
 
 		return FReply::Handled();
 	}
@@ -319,7 +319,7 @@ FReply FWaveformTransformationTrimFadeRenderer::OnMouseWheel(SWidget& OwnerWidge
 		StrongTrimFade->PostEditChange();
 		
 		EndTransaction();
-		StrongTrimFade->OnTransformationChanged.ExecuteIfBound();
+		StrongTrimFade->OnTransformationChanged.ExecuteIfBound(true);
 
 		return FReply::Handled();
 	}
@@ -426,7 +426,7 @@ FReply FWaveformTransformationTrimFadeRenderer::OnMouseButtonUp(SWidget& OwnerWi
 		SetPropertyValueDependingOnInteractionType(MouseEvent, MyGeometry);
 
 		EndTransaction();
-		StrongTrimFade->OnTransformationChanged.ExecuteIfBound();
+		StrongTrimFade->OnTransformationChanged.ExecuteIfBound(true);
 
 		TrimFadeInteractionType = ETrimFadeInteractionType::None;
 		
@@ -573,7 +573,7 @@ void FWaveformTransformationTrimFadeRenderer::ShowSelectFadeModeMenuAtCursor(con
 							}
 
 							PinnedTrimFade->PostEditChange();
-							PinnedTrimFade->OnTransformationChanged.ExecuteIfBound();
+							PinnedTrimFade->OnTransformationChanged.ExecuteIfBound(true);
 						}
 
 						TSharedPtr<SWindow> LockedFadeModeMenuWindow = InWeakFadeModeMenuWindow.Pin();
diff --git a/Engine/Source/Runtime/AudioExtensions/Private/IWaveformTransformation.cpp b/Engine/Source/Runtime/AudioExtensions/Private/IWaveformTransformation.cpp
--- a/Engine/Source/Runtime/AudioExtensions/Private/IWaveformTransformation.cpp
+++ b/Engine/Source/Runtime/AudioExtensions/Private/IWaveformTransformation.cpp
@@ -26,7 +26,8 @@ void UWaveformTransformationBase::PostEditUndo()
 {
 	Super::PostEditUndo();
 
-	OnTransformationChanged.ExecuteIfBound();
+	static const bool bMarkFileDirty = true;
+	OnTransformationChanged.ExecuteIfBound(bMarkFileDirty);
 }
 
 void UWaveformTransformationBase::NotifyPropertyChange(FProperty* Property)
diff --git a/Engine/Source/Runtime/AudioExtensions/Public/IWaveformTransformation.h b/Engine/Source/Runtime/AudioExtensions/Public/IWaveformTransformation.h
--- a/Engine/Source/Runtime/AudioExtensions/Public/IWaveformTransformation.h
+++ b/Engine/Source/Runtime/AudioExtensions/Public/IWaveformTransformation.h
@@ -163,8 +163,8 @@ public:
 	virtual bool IsEditorOnly() const override { return true; }
 
 
-	// Execute when the soundwave needs a recook to apply new transformation changes
-	DECLARE_DELEGATE(FOnTransformationChanged);
+	// Execute when the soundwave needs a recook to apply new transformation changes, bool flag if it should mark the file dirty
+	DECLARE_DELEGATE_OneParam(FOnTransformationChanged, bool);
 	FOnTransformationChanged OnTransformationChanged;
 
 	// Execute when only the rendering needs to be updated, but can wait for a recook
