[Waveform Editor] - Resetting values while a sound is playing in the Waveform Editor will stop the sound.
Ensure UpdateTransportState() is called from FWaveformEditor::NotifyPostChange for all PropertyChangedEvent.ChangeTypes except EPropertyChangeType::Interactive which is handled by a different path.
UE-306950

diff --git a/Engine/Plugins/Editor/WaveformEditor/Source/WaveformEditor/Private/WaveformEditor.cpp b/Engine/Plugins/Editor/WaveformEditor/Source/WaveformEditor/Private/WaveformEditor.cpp
--- a/Engine/Plugins/Editor/WaveformEditor/Source/WaveformEditor/Private/WaveformEditor.cpp
+++ b/Engine/Plugins/Editor/WaveformEditor/Source/WaveformEditor/Private/WaveformEditor.cpp
@@ -1832,7 +1832,8 @@ void FWaveformEditor::NotifyPostChange(const FPropertyChangedEvent& PropertyChan
 		}
 	}
 
-	if (PropertyChangedEvent.ChangeType == EPropertyChangeType::ValueSet)
+	if ((PropertyChangedEvent.ChangeType & EPropertyChangeType::ValueSet) == EPropertyChangeType::ValueSet
+	 && (PropertyChangedEvent.ChangeType & EPropertyChangeType::Interactive) != EPropertyChangeType::Interactive)
 	{
 		UpdateTransportState();
 
