Waveform Editor: Fade handle adjustments

diff --git a/Engine/Plugins/Editor/WaveformEditor/Source/WaveformTransformationsWidgets/Private/WaveformTransformationTrimFadeRenderer.cpp b/Engine/Plugins/Editor/WaveformEditor/Source/WaveformTransformationsWidgets/Private/WaveformTransformationTrimFadeRenderer.cpp
--- a/Engine/Plugins/Editor/WaveformEditor/Source/WaveformTransformationsWidgets/Private/WaveformTransformationTrimFadeRenderer.cpp
+++ b/Engine/Plugins/Editor/WaveformEditor/Source/WaveformTransformationsWidgets/Private/WaveformTransformationTrimFadeRenderer.cpp
@@ -15,6 +15,7 @@
 #include "Widgets/Input/SButton.h"
 #include "Widgets/SWidget.h"
 #include "Widgets/Text/STextBlock.h"
+#include <Blueprint/WidgetLayoutLibrary.h>
 
 FWaveformTransformationTrimFadeRenderer::~FWaveformTransformationTrimFadeRenderer()
 {
@@ -35,7 +36,7 @@ int32 FWaveformTransformationTrimFadeRenderer::DrawTrimHandles(const FGeometry&
 	const bool bRenderLowerBound = StartTimeHandleX >= 0.f;
 	const bool bRenderUpperBound = EndTimeHandleX <= AllottedGeometry.Size.X;
 
-	float HandleWidth = InteractionHandleWidth;
+	const float HandleWidth = InteractionHandleSize;
 
 	const FSlateRoundedBoxBrush RoundedBoxBrush(FLinearColor::Green, 2);
 	FPaintGeometry HandleGeometry = AllottedGeometry.ToPaintGeometry(FVector2f(HandleWidth, HandleWidth), FSlateLayoutTransform(FVector2f(StartTimeHandleX, AllottedGeometry.Size.Y - HandleWidth)));
@@ -66,10 +67,16 @@ int32 FWaveformTransformationTrimFadeRenderer::DrawTrimHandles(const FGeometry&
 int32 FWaveformTransformationTrimFadeRenderer::DrawFadeCurves(const FGeometry& AllottedGeometry, FSlateWindowElementList& OutDrawElements, int32 LayerId) const
 {
 	const FLinearColor FadeColor = FLinearColor::Gray;
-	const FSlateRoundedBoxBrush RoundedBoxBrush(FadeColor, 2);
-	const float HandleWidth = InteractionPixelXDelta;
-	const float HandleHeight = AllottedGeometry.GetLocalSize().Y * InteractionRatioYDelta;
+	const float HandleWidth = InteractionHandleSize;
+	const float HandleHeight = InteractionHandleSize;
+	const float HandleOffset = 5.0f;
 
+	const FLinearColor HandleColor = FLinearColor::White;
+	const FLinearColor HandleSelectedColor = FLinearColor::Yellow;
+
+	const bool IsFadeInSelected = IsCursorInFadeInInteractionRange(MousePosition, AllottedGeometry);
+	const bool IsFadeOutSelected = IsCursorInFadeOutInteractionRange(MousePosition, AllottedGeometry);
+	
 	if (FadeInCurvePoints.Num() > 0)
 	{
 		FSlateDrawElement::MakeLines(
@@ -82,28 +89,30 @@ int32 FWaveformTransformationTrimFadeRenderer::DrawFadeCurves(const FGeometry& A
 		);
 
 		const float	HandleStart = FadeInCurvePoints.Last(0).X;
-		const FPaintGeometry HandleGeometry = AllottedGeometry.ToPaintGeometry(FVector2f(HandleWidth, HandleHeight), FSlateLayoutTransform(FVector2f(HandleStart, 0)));
-		
+		const FPaintGeometry HandleGeometry = AllottedGeometry.ToPaintGeometry(FVector2f(HandleWidth, HandleHeight), FSlateLayoutTransform(FVector2f(HandleStart - HandleOffset, 0)));
+		const FSlateBrush* HandleBrush = FAppStyle::GetBrush("WorldBrowser.DirectionXPositive");
+	
 		FSlateDrawElement::MakeBox(
 			OutDrawElements,
-			++LayerId,
+			LayerId,
 			HandleGeometry,
-			&RoundedBoxBrush,
-			ESlateDrawEffect::None,
-			FadeColor
+			HandleBrush,
+			ESlateDrawEffect::DisabledEffect,
+			IsFadeInSelected ? HandleSelectedColor : HandleColor
 		);
 	}
 	else
 	{
-		FPaintGeometry HandleGeometry = AllottedGeometry.ToPaintGeometry(FVector2f(HandleWidth, HandleHeight), FSlateLayoutTransform(FVector2f(StartTimeHandleX, 0)));
+		FPaintGeometry HandleGeometry = AllottedGeometry.ToPaintGeometry(FVector2f(HandleWidth, HandleHeight), FSlateLayoutTransform(FVector2f(StartTimeHandleX - HandleOffset, 0)));
+		const FSlateBrush* HandleBrush = FAppStyle::GetBrush("WorldBrowser.DirectionXPositive");
 
 		FSlateDrawElement::MakeBox(
 			OutDrawElements,
-			++LayerId,
+			LayerId,
 			HandleGeometry,
-			&RoundedBoxBrush,
-			ESlateDrawEffect::None,
-			FadeColor
+			HandleBrush,
+			ESlateDrawEffect::DisabledEffect,
+			IsFadeInSelected ? HandleSelectedColor : HandleColor
 		);
 	}
 
@@ -119,28 +128,30 @@ int32 FWaveformTransformationTrimFadeRenderer::DrawFadeCurves(const FGeometry& A
 		);
 
 		const float	HandleStart = FadeOutCurvePoints[0].X - HandleWidth;
-		const FPaintGeometry HandleGeometry = AllottedGeometry.ToPaintGeometry(FVector2f(HandleWidth, HandleHeight), FSlateLayoutTransform(FVector2f(HandleStart, 0)));
+		const FPaintGeometry HandleGeometry = AllottedGeometry.ToPaintGeometry(FVector2f(HandleWidth, HandleHeight), FSlateLayoutTransform(FVector2f(HandleStart + HandleOffset, 0)));
+		const FSlateBrush* HandleBrush = FAppStyle::GetBrush("WorldBrowser.DirectionXNegative");
 
 		FSlateDrawElement::MakeBox(
 			OutDrawElements,
-			++LayerId,
+			LayerId,
 			HandleGeometry,
-			&RoundedBoxBrush,
-			ESlateDrawEffect::None,
-			FadeColor
+			HandleBrush,
+			ESlateDrawEffect::DisabledEffect,
+			IsFadeOutSelected ? HandleSelectedColor : HandleColor
 		);
 	}
 	else
 	{
-		FPaintGeometry HandleGeometry = AllottedGeometry.ToPaintGeometry(FVector2f(HandleWidth, HandleHeight), FSlateLayoutTransform(FVector2f(EndTimeHandleX - HandleWidth, 0)));
+		FPaintGeometry HandleGeometry = AllottedGeometry.ToPaintGeometry(FVector2f(HandleWidth, HandleHeight), FSlateLayoutTransform(FVector2f(EndTimeHandleX - HandleWidth + HandleOffset, 0)));
+		const FSlateBrush* HandleBrush = FAppStyle::GetBrush("WorldBrowser.DirectionXNegative");
 
 		FSlateDrawElement::MakeBox(
 			OutDrawElements,
-			++LayerId,
+			LayerId,
 			HandleGeometry,
-			&RoundedBoxBrush,
-			ESlateDrawEffect::None,
-			FadeColor
+			HandleBrush,
+			ESlateDrawEffect::DisabledEffect,
+			IsFadeOutSelected ? HandleSelectedColor : HandleColor
 		);
 	}
 
@@ -357,26 +368,29 @@ void FWaveformTransformationTrimFadeRenderer::Tick(const FGeometry& AllottedGeom
 	StartTimeHandleX = FirstFrame * PixelsPerFrame;
 	EndTimeHandleX = EndFrame * PixelsPerFrame;
 
+	FVector2D MouseAbsolutePosition = UWidgetLayoutLibrary::GetMousePositionOnPlatform();
+	MousePosition = AllottedGeometry.AbsoluteToLocal(MouseAbsolutePosition);
+
 	GenerateFadeCurves(AllottedGeometry);
 	UpdateInteractionRange(AllottedGeometry);
 }
 
 void FWaveformTransformationTrimFadeRenderer::UpdateInteractionRange(const FGeometry& AllottedGeometry)
 {
-	StartTimeInteractionXRange.SetLowerBoundValue(StartTimeHandleX - InteractionHandleWidth);
-	StartTimeInteractionXRange.SetUpperBoundValue(StartTimeHandleX + InteractionHandleWidth);
-	StartTimeInteractionYRange.SetLowerBoundValue(AllottedGeometry.GetLocalSize().Y - InteractionHandleWidth);
+	StartTimeInteractionXRange.SetLowerBoundValue(StartTimeHandleX - InteractionHandleSize);
+	StartTimeInteractionXRange.SetUpperBoundValue(StartTimeHandleX + InteractionHandleSize);
+	StartTimeInteractionYRange.SetLowerBoundValue(AllottedGeometry.GetLocalSize().Y - InteractionHandleSize);
 	StartTimeInteractionYRange.SetUpperBoundValue(AllottedGeometry.GetLocalSize().Y);
 
-	EndTimeInteractionXRange.SetLowerBoundValue(EndTimeHandleX - InteractionHandleWidth);
-	EndTimeInteractionXRange.SetUpperBoundValue(EndTimeHandleX + InteractionHandleWidth);
-	EndTimeInteractionYRange.SetLowerBoundValue(AllottedGeometry.GetLocalSize().Y - InteractionHandleWidth);
+	EndTimeInteractionXRange.SetLowerBoundValue(EndTimeHandleX - InteractionHandleSize);
+	EndTimeInteractionXRange.SetUpperBoundValue(EndTimeHandleX + InteractionHandleSize);
+	EndTimeInteractionYRange.SetLowerBoundValue(AllottedGeometry.GetLocalSize().Y - InteractionHandleSize);
 	EndTimeInteractionYRange.SetUpperBoundValue(AllottedGeometry.GetLocalSize().Y);
 
-	FadeInInteractionXRange.SetLowerBoundValue(FadeInEndX - InteractionHandleWidth);
-	FadeInInteractionXRange.SetUpperBoundValue(FadeInEndX + InteractionHandleWidth);
-	FadeOutInteractionXRange.SetLowerBoundValue(FadeOutStartX - InteractionHandleWidth);
-	FadeOutInteractionXRange.SetUpperBoundValue(FadeOutStartX + InteractionHandleWidth);
+	FadeInInteractionXRange.SetLowerBoundValue(FadeInEndX - InteractionHandleSize);
+	FadeInInteractionXRange.SetUpperBoundValue(FadeInEndX + InteractionHandleSize);
+	FadeOutInteractionXRange.SetLowerBoundValue(FadeOutStartX - InteractionHandleSize);
+	FadeOutInteractionXRange.SetUpperBoundValue(FadeOutStartX + InteractionHandleSize);
 }
 
 FReply FWaveformTransformationTrimFadeRenderer::OnMouseButtonDown(SWidget& OwnerWidget, const FGeometry& MyGeometry, const FPointerEvent& MouseEvent)
@@ -737,13 +751,13 @@ FWaveformTransformationTrimFadeRenderer::ETrimFadeInteractionType FWaveformTrans
 bool FWaveformTransformationTrimFadeRenderer::IsCursorInFadeInInteractionRange(const FVector2D& InLocalCursorPosition, const FGeometry& WidgetGeometry) const
 {
 	return FadeInInteractionXRange.Contains(InLocalCursorPosition.X)
-		&& InLocalCursorPosition.Y < WidgetGeometry.GetLocalSize().Y* InteractionRatioYDelta;
+		&& InLocalCursorPosition.Y < InteractionHandleSize;
 }
 
 bool FWaveformTransformationTrimFadeRenderer::IsCursorInFadeOutInteractionRange(const FVector2D& InLocalCursorPosition, const FGeometry& WidgetGeometry) const
 {
 	return FadeOutInteractionXRange.Contains(InLocalCursorPosition.X)
-		&& InLocalCursorPosition.Y < WidgetGeometry.GetLocalSize().Y* InteractionRatioYDelta;
+		&& InLocalCursorPosition.Y < InteractionHandleSize;
 }
 
 FText FWaveformTransformationTrimFadeRenderer::GetPropertyEditedByCurrentInteraction() const
diff --git a/Engine/Plugins/Editor/WaveformEditor/Source/WaveformTransformationsWidgets/Private/WaveformTransformationTrimFadeRenderer.h b/Engine/Plugins/Editor/WaveformEditor/Source/WaveformTransformationsWidgets/Private/WaveformTransformationTrimFadeRenderer.h
--- a/Engine/Plugins/Editor/WaveformEditor/Source/WaveformTransformationsWidgets/Private/WaveformTransformationTrimFadeRenderer.h
+++ b/Engine/Plugins/Editor/WaveformEditor/Source/WaveformTransformationsWidgets/Private/WaveformTransformationTrimFadeRenderer.h
@@ -51,7 +51,8 @@ private:
 	TRange<float> FadeOutInteractionXRange = TRange<float>::Inclusive(0.f, 0.f);
 
 	double PixelsPerFrame = 0.0;
-	const float InteractionHandleWidth = 20.0f;
+	const float InteractionHandleSize = 30.0f;
+	FVector2D MousePosition;
 
 	enum class ETrimFadeInteractionType : uint8
 	{
