[Waveform Editor] - Cannot Undo/Redo Fade Function Duration and Curve parameters
[Waveform Editor] - Unable to use original Fade In/Out Curve and Duration parameters

diff --git a/Engine/Plugins/Editor/WaveformEditor/Source/WaveformEditor/Private/WaveformEditor.cpp b/Engine/Plugins/Editor/WaveformEditor/Source/WaveformEditor/Private/WaveformEditor.cpp
--- a/Engine/Plugins/Editor/WaveformEditor/Source/WaveformEditor/Private/WaveformEditor.cpp
+++ b/Engine/Plugins/Editor/WaveformEditor/Source/WaveformEditor/Private/WaveformEditor.cpp
@@ -261,22 +261,13 @@ FWaveformEditor::~FWaveformEditor()
 
 		if (FMath::IsNearlyEqual(TrimFadeTransformation->StartTime, 0, ErrorTolerance) && FMath::IsNearlyEqual(TrimFadeTransformation->EndTime, SoundWave->Duration, ErrorTolerance))
 		{
-			bool bRemoveTrimFade = true;
-
-			if ((TrimFadeTransformation->FadeFunctions.FadeIn && !FMath::IsNearlyEqual(TrimFadeTransformation->FadeFunctions.FadeIn->Duration, 0, ErrorTolerance))
-				|| (TrimFadeTransformation->FadeFunctions.FadeOut && !FMath::IsNearlyEqual(TrimFadeTransformation->FadeFunctions.FadeOut->Duration, 0, ErrorTolerance)))
-			{
-				// If any FadeFunctions are present with a non-zero duration, do not remove the transformation
-				bRemoveTrimFade = false;
-			}
-			
-			PRAGMA_DISABLE_DEPRECATION_WARNINGS
-			// If no FadeFunctions are present, check the deprecated fade properties before removing the trimfade
-			if (bRemoveTrimFade && FMath::IsNearlyEqual(TrimFadeTransformation->StartFadeTime, 0, ErrorTolerance) && FMath::IsNearlyEqual(TrimFadeTransformation->EndFadeTime, 0, ErrorTolerance))
+			if ((TrimFadeTransformation->FadeFunctions.FadeIn && FMath::IsNearlyEqual(TrimFadeTransformation->FadeFunctions.FadeIn->Duration, 0, ErrorTolerance)
+				&& TrimFadeTransformation->FadeFunctions.FadeOut && FMath::IsNearlyEqual(TrimFadeTransformation->FadeFunctions.FadeOut->Duration, 0, ErrorTolerance))
+				|| (TrimFadeTransformation->FadeFunctions.FadeIn == nullptr && TrimFadeTransformation->FadeFunctions.FadeOut == nullptr))
 			{
+				// If FadeIn and FadeOut Functions are present and they both have a zero duration or if both fade functions are null, remove the transformation
 				SoundWave->Transformations.Remove(TrimFadeTransformation);
 			}
-			PRAGMA_ENABLE_DEPRECATION_WARNINGS
 		}
 	}
 
@@ -377,9 +368,9 @@ TObjectPtr<UWaveformTransformationTrimFade> FWaveformEditor::GetOrAddTrimFadeTra
 		}
 
 		// Add default fade curves
-		TrimFadeTransformation->FadeFunctions.FadeIn = NewObject<UFadeFunction>(TrimFadeTransformation.Get(), FadeInType.Get());
+		TrimFadeTransformation->FadeFunctions.FadeIn = NewObject<UFadeFunction>(TrimFadeTransformation.Get(), FadeInType.Get(), NAME_None, RF_Transactional);
 		check(TrimFadeTransformation->FadeFunctions.FadeIn);
-		TrimFadeTransformation->FadeFunctions.FadeOut = NewObject<UFadeFunction>(TrimFadeTransformation.Get(), FadeOutType.Get());
+		TrimFadeTransformation->FadeFunctions.FadeOut = NewObject<UFadeFunction>(TrimFadeTransformation.Get(), FadeOutType.Get(), NAME_None, RF_Transactional);
 		check(TrimFadeTransformation->FadeFunctions.FadeOut);
 
 		TrimFadeTransformation->FadeFunctions.FadeIn->Duration = 0.f;
@@ -470,24 +461,21 @@ void FWaveformEditor::ToggleFadeIn()
 {
 	check(SoundWave);
 	check(GEditor != nullptr);
-	GEditor->BeginTransaction(TEXT("PropertyEditor"), LOCTEXT("SetNumericPropertyTransactionToggleFadeIn", "ToggleFadeIn"), nullptr);
-	SoundWave->Modify();
 
 	TObjectPtr<UWaveformTransformationTrimFade> TrimFadeTransformation = GetOrAddTrimFadeTransformation();
 	check(TrimFadeTransformation);
 
-	TrimFadeTransformation->Modify();
-	
 	if (TrimFadeTransformation->FadeFunctions.FadeIn)
 	{
+		GEditor->BeginTransaction(TEXT("PropertyEditor"), LOCTEXT("SetNumericPropertyTransactionToggleFadeIn", "ToggleFadeIn"), nullptr);
+		SoundWave->Modify();
+		TrimFadeTransformation->Modify();
+		TrimFadeTransformation->FadeFunctions.FadeIn->Modify();
+		
 		if (TrimFadeTransformation->FadeFunctions.FadeIn->Duration == 0.0f)
 		{
 			check(DefaultFadeInAmount >= 0.0f);
 
-			PRAGMA_DISABLE_DEPRECATION_WARNINGS
-			check(FadeInCurve >= UWaveformTransformationTrimFade::MinFadeCurve && FadeInCurve <= UWaveformTransformationTrimFade::MaxFadeCurve);
-			PRAGMA_ENABLE_DEPRECATION_WARNINGS
-
 			if (CachedFadeInAmount <= 0.0f)
 			{
 				CachedFadeInAmount = DefaultFadeInAmount;
@@ -502,45 +490,10 @@ void FWaveformEditor::ToggleFadeIn()
 			//No need to set FadeCurve when toggling off. This will preserve user curve data if they later manually change the StartFadeTime.
 			TrimFadeTransformation->FadeFunctions.FadeIn->Duration = 0.0f;
 		}
+		
+		NotifyPostTransformationChange(EPropertyChangeType::ValueSet);
+		GEditor->EndTransaction();
 	}
-	else // TrimFadeTransformation->FadeFunctions.FadeIn is nullptr
-	{
-		PRAGMA_DISABLE_DEPRECATION_WARNINGS
-		//if fade duration is 0, assume there is no fade present and toggle on the fade with the curve data selected in the menu
-		if (TrimFadeTransformation->StartFadeTime == 0.0f)
-		{
-			// No cached amount means we can instantiate a new FadeIn Function object
-			if (CachedFadeInAmount == 0.f)
-			{
-				TrimFadeTransformation->FadeFunctions.FadeIn = NewObject<UFadeFunction>(TrimFadeTransformation.Get(), FadeInType.Get());
-			}
-			else
-			{
-				check(DefaultFadeInAmount >= 0.0f);
-				check(FadeInCurve >= UWaveformTransformationTrimFade::MinFadeCurve && FadeInCurve <= UWaveformTransformationTrimFade::MaxFadeCurve);
-
-				if (CachedFadeInAmount <= 0.0f)
-				{
-					CachedFadeInAmount = DefaultFadeInAmount;
-				}
-
-				TrimFadeTransformation->StartFadeTime = CachedFadeInAmount;
-				TrimFadeTransformation->StartFadeCurve = FadeInCurve;
-			}
-		}
-		else //if fade duration is not 0, assume there is a fade applied and toggle fade "off" by setting fade time to 0
-		{
-			CachedFadeInAmount = TrimFadeTransformation->StartFadeTime;
-
-			//No need to set FadeCurve when toggling off. This will preserve user curve data if they later manually change the StartFadeTime.
-			TrimFadeTransformation->StartFadeTime = 0.0f;
-		}
-		PRAGMA_ENABLE_DEPRECATION_WARNINGS
-	}
-
-	NotifyPostTransformationChange(EPropertyChangeType::ValueSet);
-
-	GEditor->EndTransaction();
 }
 
 bool FWaveformEditor::CanFadeIn()
@@ -552,24 +505,21 @@ void FWaveformEditor::ToggleFadeOut()
 {
 	check(SoundWave);
 	check(GEditor != nullptr);
-	GEditor->BeginTransaction(TEXT("PropertyEditor"), LOCTEXT("SetNumericPropertyTransactionToggleFadeOut", "ToggleFadeOut"), nullptr);
-	SoundWave->Modify();
 
 	TObjectPtr<UWaveformTransformationTrimFade> TrimFadeTransformation = GetOrAddTrimFadeTransformation();
 	check(TrimFadeTransformation);
 
-	TrimFadeTransformation->Modify();
-
 	if (TrimFadeTransformation->FadeFunctions.FadeOut)
 	{
+		GEditor->BeginTransaction(TEXT("PropertyEditor"), LOCTEXT("SetNumericPropertyTransactionToggleFadeOut", "ToggleFadeOut"), nullptr);
+		SoundWave->Modify();
+		TrimFadeTransformation->Modify();
+		TrimFadeTransformation->FadeFunctions.FadeOut->Modify();
+		
 		if (TrimFadeTransformation->FadeFunctions.FadeOut->Duration == 0.0f)
 		{
 			check(DefaultFadeOutAmount >= 0.0f);
 
-			PRAGMA_DISABLE_DEPRECATION_WARNINGS
-			check(FadeOutCurve >= UWaveformTransformationTrimFade::MinFadeCurve && FadeOutCurve <= UWaveformTransformationTrimFade::MaxFadeCurve);
-			PRAGMA_ENABLE_DEPRECATION_WARNINGS
-
 			if (CachedFadeOutAmount <= 0.0f)
 			{
 				CachedFadeOutAmount = DefaultFadeOutAmount;
@@ -584,44 +534,10 @@ void FWaveformEditor::ToggleFadeOut()
 			//No need to set FadeCurve when toggling off. This will preserve user curve data if they later manually change the StartFadeTime.
 			TrimFadeTransformation->FadeFunctions.FadeOut->Duration = 0.0f;
 		}
-	}
-	else // TrimFadeTransformation->FadeFunctions.FadeOut is nullptr
-	{
-		PRAGMA_DISABLE_DEPRECATION_WARNINGS
-		//if fade duration is 0, assume there is no fade present and toggle on the fade with the curve data selected in the menu
-		if (TrimFadeTransformation->EndFadeTime == 0.0f)
-		{
-			// No cached amount means we can instantiate a new FadeIn Function object
-			if (CachedFadeOutAmount == 0.f)
-			{
-				TrimFadeTransformation->FadeFunctions.FadeOut = NewObject<UFadeFunction>(TrimFadeTransformation.Get(), FadeOutType.Get());
-			}
-			else
-			{
-				check(DefaultFadeOutAmount >= 0.0f);
-				check(FadeOutCurve >= UWaveformTransformationTrimFade::MinFadeCurve && FadeOutCurve <= UWaveformTransformationTrimFade::MaxFadeCurve);
-
-				if (CachedFadeOutAmount <= 0.0f)
-				{
-					CachedFadeOutAmount = DefaultFadeOutAmount;
-				}
 
-				TrimFadeTransformation->EndFadeTime = CachedFadeOutAmount;
-				TrimFadeTransformation->EndFadeCurve = FadeOutCurve;
-			}
-		}
-		else //if fade duration is not 0, assume there is a fade applied and toggle fade "off" by setting fade time to 0
-		{
-			CachedFadeOutAmount = TrimFadeTransformation->EndFadeTime;
-
-			//No need to set FadeCurve when toggling off. This will preserve user curve data if they later manually change the EndFadeTime.
-			TrimFadeTransformation->EndFadeTime = 0.0f;
-		}
-		PRAGMA_ENABLE_DEPRECATION_WARNINGS
+		NotifyPostTransformationChange(EPropertyChangeType::ValueSet);
+		GEditor->EndTransaction();
 	}
-
-	NotifyPostTransformationChange(EPropertyChangeType::ValueSet);
-	GEditor->EndTransaction();
 }
 
 bool FWaveformEditor::CanFadeOut()
@@ -1024,7 +940,10 @@ void FWaveformEditor::ExecuteOverwriteTransformations()
 {
 	for (TObjectPtr<UWaveformTransformationBase> Transformation : SoundWave->Transformations)
 	{
-		Transformation->OverwriteTransformation();
+		if (Transformation)
+		{
+			Transformation->OverwriteTransformation();
+		}
 	}
 }
 
@@ -1510,7 +1429,6 @@ bool FWaveformEditor::BindCommands()
 		Commands.FadeInLinear,
 		FExecuteAction::CreateLambda([this] 
 			{ 
-				FadeInCurve = UWaveformTransformationTrimFade::FadeModeToCurveValueMap[EWaveEditorFadeMode::Linear];
 				FadeInType = UWaveformTransformationTrimFade::FadeModeToFadeFunctionMap[EWaveEditorFadeMode::Linear];
 			}),
 		FCanExecuteAction::CreateLambda([this] { return true; }),
@@ -1520,7 +1438,6 @@ bool FWaveformEditor::BindCommands()
 		Commands.FadeInExponential,
 		FExecuteAction::CreateLambda([this] 
 			{ 
-				FadeInCurve = UWaveformTransformationTrimFade::FadeModeToCurveValueMap[EWaveEditorFadeMode::Exponential];
 				FadeInType = UWaveformTransformationTrimFade::FadeModeToFadeFunctionMap[EWaveEditorFadeMode::Exponential];
 			}),
 		FCanExecuteAction::CreateLambda([this] { return true; }),
@@ -1530,7 +1447,6 @@ bool FWaveformEditor::BindCommands()
 		Commands.FadeInLogarithmic,
 		FExecuteAction::CreateLambda([this] 
 			{ 
-				FadeInCurve = UWaveformTransformationTrimFade::FadeModeToCurveValueMap[EWaveEditorFadeMode::Logarithmic];
 				FadeInType = UWaveformTransformationTrimFade::FadeModeToFadeFunctionMap[EWaveEditorFadeMode::Logarithmic];
 			}),
 		FCanExecuteAction::CreateLambda([this] { return true; }),
@@ -1540,7 +1456,6 @@ bool FWaveformEditor::BindCommands()
 		Commands.FadeInSigmoid,
 		FExecuteAction::CreateLambda([this] 
 			{ 
-				FadeInCurve = UWaveformTransformationTrimFade::FadeModeToCurveValueMap[EWaveEditorFadeMode::Sigmoid];
 				FadeInType = UWaveformTransformationTrimFade::FadeModeToFadeFunctionMap[EWaveEditorFadeMode::Sigmoid];
 			}),
 		FCanExecuteAction::CreateLambda([this] { return true; }),
@@ -1555,7 +1470,6 @@ bool FWaveformEditor::BindCommands()
 		Commands.FadeOutLinear,
 		FExecuteAction::CreateLambda([this] 
 			{ 
-				FadeOutCurve = UWaveformTransformationTrimFade::FadeModeToCurveValueMap[EWaveEditorFadeMode::Linear];
 				FadeOutType = UWaveformTransformationTrimFade::FadeModeToFadeFunctionMap[EWaveEditorFadeMode::Linear];
 			}),
 		FCanExecuteAction::CreateLambda([this] {return true; }),
@@ -1565,7 +1479,6 @@ bool FWaveformEditor::BindCommands()
 		Commands.FadeOutExponential,
 		FExecuteAction::CreateLambda([this] 
 			{ 
-				FadeOutCurve = UWaveformTransformationTrimFade::FadeModeToCurveValueMap[EWaveEditorFadeMode::Exponential];
 				FadeOutType = UWaveformTransformationTrimFade::FadeModeToFadeFunctionMap[EWaveEditorFadeMode::Exponential];
 			}),
 		FCanExecuteAction::CreateLambda([this] {return true; }),
@@ -1575,7 +1488,6 @@ bool FWaveformEditor::BindCommands()
 		Commands.FadeOutLogarithmic,
 		FExecuteAction::CreateLambda([this] 
 			{ 
-				FadeOutCurve = UWaveformTransformationTrimFade::FadeModeToCurveValueMap[EWaveEditorFadeMode::Logarithmic];
 				FadeOutType = UWaveformTransformationTrimFade::FadeModeToFadeFunctionMap[EWaveEditorFadeMode::Logarithmic];
 			}),
 		FCanExecuteAction::CreateLambda([this] {return true; }),
@@ -1585,7 +1497,6 @@ bool FWaveformEditor::BindCommands()
 		Commands.FadeOutSigmoid,
 		FExecuteAction::CreateLambda([this] 
 			{ 
-				FadeOutCurve = UWaveformTransformationTrimFade::FadeModeToCurveValueMap[EWaveEditorFadeMode::Sigmoid];
 				FadeOutType = UWaveformTransformationTrimFade::FadeModeToFadeFunctionMap[EWaveEditorFadeMode::Sigmoid];
 			}),
 		FCanExecuteAction::CreateLambda([this] {return true; }),
@@ -1698,14 +1609,23 @@ void FWaveformEditor::OnAssetReimport(UObject* ReimportedObject, bool bSuccessfu
 
 	if (ReimportedObject->IsA<USoundWave>() && ReimportedObject->GetPathName() == SoundWave->GetPathName())
 	{
+		check(SoundWave);
+
 		//If the waveform editor is open, updating Transformations on reimport will show
 		// the effects of the overwritten transformations immediately instead of after a
 		// change is made by the user.
 		TransformationChainConfig = SoundWave->UpdateTransformations();
 		CreateWaveformView();
+
+		check(WaveformView.DataProvider);
+		check(WaveformView.ViewWidget);
 		WaveformView.DataProvider->RequestSequenceView(TransportCoordinator->GetDisplayRange());
 		WaveformView.ViewWidget->SetPlayheadRatio(TransportCoordinator->GetFocusPoint());
-		this->TabManager->FindExistingLiveTab(WaveformDisplayTabId)->SetContent(WaveformView.ViewWidget.ToSharedRef());
+
+		if (TSharedPtr<SDockTab> LiveTab = this->TabManager->FindExistingLiveTab(WaveformDisplayTabId))
+		{
+			LiveTab->SetContent(WaveformView.ViewWidget.ToSharedRef());
+		}
 	}
 }
 
@@ -1917,19 +1837,19 @@ void FWaveformEditor::PostRedo(bool bSuccess)
 
 bool FWaveformEditor::MatchesContext(const FTransactionContext& InContext, const TArray<TPair<UObject*, FTransactionObjectEvent>>& TransactionObjectContexts) const
 {
-	bool bShoouldUndo = false;
+	bool bShouldUndo = false;
 	for (const TPair<UObject*, FTransactionObjectEvent>& TransactionObjectPair : TransactionObjectContexts)
 	{
-		UObject* Object = TransactionObjectPair.Key;
-		const UClass* ObjectClass = Object->GetClass();
-
-		if (ObjectClass)
+		if (const UObject* Object = TransactionObjectPair.Key)
 		{
-			bShoouldUndo = ObjectClass->IsChildOf(UWaveformTransformationBase::StaticClass()) || ObjectClass->IsChildOf(USoundWave::StaticClass());
+			if (const UClass* ObjectClass = Object->GetClass())
+			{
+				bShouldUndo |= ObjectClass->IsChildOf(UWaveformTransformationBase::StaticClass()) || ObjectClass->IsChildOf(USoundWave::StaticClass());
+			}
 		}
 	}
 
-	return bShoouldUndo;
+	return bShouldUndo;
 }
 
 void FWaveformEditor::InitToolMenuContext(FToolMenuContext& MenuContext)
diff --git a/Engine/Plugins/Editor/WaveformEditor/Source/WaveformEditor/Public/WaveformEditor.h b/Engine/Plugins/Editor/WaveformEditor/Source/WaveformEditor/Public/WaveformEditor.h
--- a/Engine/Plugins/Editor/WaveformEditor/Source/WaveformEditor/Public/WaveformEditor.h
+++ b/Engine/Plugins/Editor/WaveformEditor/Source/WaveformEditor/Public/WaveformEditor.h
@@ -179,12 +179,10 @@ private:
 
 	static constexpr float DefaultFadeInAmount = 0.5f;
 	float CachedFadeInAmount = 0.f;
-	float FadeInCurve = UWaveformTransformationTrimFade::FadeModeToCurveValueMap[EWaveEditorFadeMode::Exponential];
 	TSubclassOf<UFadeFunction> FadeInType = UWaveformTransformationTrimFade::FadeModeToFadeFunctionMap[EWaveEditorFadeMode::Exponential];
 
 	static constexpr float DefaultFadeOutAmount = 0.5f;
 	float CachedFadeOutAmount = 0.f;
-	float FadeOutCurve = UWaveformTransformationTrimFade::FadeModeToCurveValueMap[EWaveEditorFadeMode::Exponential];
 	TSubclassOf<UFadeFunction> FadeOutType = UWaveformTransformationTrimFade::FadeModeToFadeFunctionMap[EWaveEditorFadeMode::Exponential];
 
 	// Tracking bool to ensure we detect wave cue changes and regenerate transformations
diff --git a/Engine/Plugins/Editor/WaveformEditor/Source/WaveformTransformations/Private/WaveformTransformationTrimFade.cpp b/Engine/Plugins/Editor/WaveformEditor/Source/WaveformTransformations/Private/WaveformTransformationTrimFade.cpp
--- a/Engine/Plugins/Editor/WaveformEditor/Source/WaveformTransformations/Private/WaveformTransformationTrimFade.cpp
+++ b/Engine/Plugins/Editor/WaveformEditor/Source/WaveformTransformations/Private/WaveformTransformationTrimFade.cpp
@@ -1,6 +1,8 @@
 // Copyright Epic Games, Inc. All Rights Reserved.
 
 #include "WaveformTransformationTrimFade.h"
+
+#include "Editor.h"
 #include "Sound/SoundWave.h"
 
 #include UE_INLINE_GENERATED_CPP_BY_NAME(WaveformTransformationTrimFade)
@@ -9,15 +11,29 @@ namespace
 {
 	// Makes the two functions for SCurve meet at (0.5, 0.5)
 	const double XMultiplier = FMath::Sqrt(10.0) / 5.0;
-}
 
-const TMap<EWaveEditorFadeMode, TSubclassOf<UFadeFunction>> UWaveformTransformationTrimFade::FadeModeToFadeFunctionMap =
-{
-	{EWaveEditorFadeMode::Linear, UFadeFunctionLinear::StaticClass()},
-	{EWaveEditorFadeMode::Exponential, UFadeCurveFunctionExponential::StaticClass()},
-	{EWaveEditorFadeMode::Logarithmic, UFadeCurveFunctionLogarithmic::StaticClass()},
-	{EWaveEditorFadeMode::Sigmoid, UFadeCurveFunctionSigmoid::StaticClass()}
-};
+	UE_DEPRECATED(5.7, "Individual fade properties have been deprecated. Use FadeFunctions::FadeIn::GetFadeInCurveValue instead.")
+	void SetFadeCurve(TObjectPtr<UFadeFunction> FadeFunction, float FadeCurve, float Sharpness)
+	{
+		if (FadeFunction->IsA<UFadeCurveFunction>())
+		{
+			TObjectPtr<UFadeCurveFunction> CurveFunction = CastChecked<UFadeCurveFunction>(FadeFunction);
+
+			if (CurveFunction)
+			{
+				if (!CurveFunction->IsA<UFadeCurveFunctionSigmoid>())
+				{
+					CurveFunction->SetFadeCurve(FadeCurve);
+				}
+				else
+				{
+					CurveFunction->SetFadeCurve(Sharpness);
+				}
+			}
+			
+		}
+	}
+}
 
 PRAGMA_DISABLE_DEPRECATION_WARNINGS
 const TMap<EWaveEditorFadeMode, float> UWaveformTransformationTrimFade::FadeModeToCurveValueMap =
@@ -29,20 +45,35 @@ const TMap<EWaveEditorFadeMode, float> UWaveformTransformationTrimFade::FadeMode
 };
 PRAGMA_ENABLE_DEPRECATION_WARNINGS
 
+const TMap<EWaveEditorFadeMode, TSubclassOf<UFadeFunction>> UWaveformTransformationTrimFade::FadeModeToFadeFunctionMap =
+{
+	{EWaveEditorFadeMode::Linear, UFadeFunctionLinear::StaticClass()},
+	{EWaveEditorFadeMode::Exponential, UFadeCurveFunctionExponential::StaticClass()},
+	{EWaveEditorFadeMode::Logarithmic, UFadeCurveFunctionLogarithmic::StaticClass()},
+	{EWaveEditorFadeMode::Sigmoid, UFadeCurveFunctionSigmoid::StaticClass()}
+};
+
 const float UWaveformTransformationTrimFade::MinFadeCurve = -0.1f;
 const float UWaveformTransformationTrimFade::MaxFadeCurve = 10.f;
 
+
 namespace WaveformTransformationTrimFadeNames
 {
 	static FLazyName StartTimeName(GET_MEMBER_NAME_CHECKED(UWaveformTransformationTrimFade, StartTime));
 	static FLazyName EndTimeName(GET_MEMBER_NAME_CHECKED(UWaveformTransformationTrimFade, EndTime));
 }
 
-static void ApplyFadeIn(Audio::FAlignedFloatBuffer& InputAudio, const FFadeFunctionData& FadeFunctions, const float FadeLength, const float FadeCurve, const float SCurveSharpness, const int32 NumChannels, const float SampleRate)
+static void ApplyFadeIn(Audio::FAlignedFloatBuffer& InputAudio, const FFadeFunctionData& FadeFunctions, const int32 NumChannels, const float SampleRate)
 {
+	if (FadeFunctions.FadeIn == nullptr)
+	{
+		return;
+	}
+
 	check(NumChannels > 0);
 	check(SampleRate > 0.f);
-	const float FadeDuration = FadeFunctions.FadeIn ? FadeFunctions.FadeIn->Duration : FadeLength;
+
+	const float FadeDuration = FadeFunctions.FadeIn->Duration;
 	const int32 FadeNumFrames = FMath::Min(FadeDuration * SampleRate, InputAudio.Num() / NumChannels);
 
 	if (InputAudio.Num() < NumChannels || FadeDuration < SMALL_NUMBER || FadeNumFrames == 0)
@@ -58,10 +89,7 @@ static void ApplyFadeIn(Audio::FAlignedFloatBuffer& InputAudio, const FFadeFunct
 		check(FadeFraction <= 1.0f);
 		check(FadeFraction >= 0.0f);
 
-		PRAGMA_DISABLE_DEPRECATION_WARNINGS
-		const float EnvValue = FadeFunctions.FadeIn ? FadeFunctions.FadeIn->GetFadeInCurveValue(FadeFraction) :
-		UWaveformTransformationTrimFade::GetFadeInCurveValue(FadeCurve, FadeFraction, SCurveSharpness);
-		PRAGMA_ENABLE_DEPRECATION_WARNINGS
+		const float EnvValue = FadeFunctions.FadeIn->GetFadeInCurveValue(FadeFraction);
 
 		for (int32 ChannelIt = 0; ChannelIt < NumChannels; ++ChannelIt)
 		{
@@ -72,11 +100,17 @@ static void ApplyFadeIn(Audio::FAlignedFloatBuffer& InputAudio, const FFadeFunct
 	}
 }
 
-static void ApplyFadeOut(Audio::FAlignedFloatBuffer& InputAudio, const FFadeFunctionData& FadeFunctions, const float FadeLength, const float FadeCurve, const float SCurveSharpness, const int32 NumChannels, const float SampleRate)
+static void ApplyFadeOut(Audio::FAlignedFloatBuffer& InputAudio, const FFadeFunctionData& FadeFunctions, const int32 NumChannels, const float SampleRate)
 {
+	if (FadeFunctions.FadeOut == nullptr)
+	{
+		return;
+	}
+
 	check(NumChannels > 0);
 	check(SampleRate > 0.f);
-	const float FadeDuration = FadeFunctions.FadeOut ? FadeFunctions.FadeOut->Duration : FadeLength;
+
+	const float FadeDuration = FadeFunctions.FadeOut->Duration;
 	const int32 FadeNumFrames = FMath::Min(FadeDuration * SampleRate, InputAudio.Num() / NumChannels);
 
 	if (InputAudio.Num() < NumChannels || FadeDuration < SMALL_NUMBER || FadeNumFrames == 0)
@@ -93,10 +127,7 @@ static void ApplyFadeOut(Audio::FAlignedFloatBuffer& InputAudio, const FFadeFunc
 		check(FadeFraction <= 1.0f);
 		check(FadeFraction >= 0.0f);
 
-		PRAGMA_DISABLE_DEPRECATION_WARNINGS
-		const float EnvValue = FadeFunctions.FadeOut ? FadeFunctions.FadeOut->GetFadeOutCurveValue(FadeFraction) :
-		UWaveformTransformationTrimFade::GetFadeOutCurveValue(FadeCurve, FadeFraction, SCurveSharpness);
-		PRAGMA_ENABLE_DEPRECATION_WARNINGS
+		const float EnvValue = FadeFunctions.FadeOut->GetFadeOutCurveValue(FadeFraction);
 
 		for (int32 ChannelIt = 0; ChannelIt < NumChannels; ++ChannelIt)
 		{
@@ -106,18 +137,11 @@ static void ApplyFadeOut(Audio::FAlignedFloatBuffer& InputAudio, const FFadeFunc
 		InputPtr += NumChannels;
 	}
 }
-PRAGMA_DISABLE_DEPRECATION_WARNINGS
-FWaveTransformationTrimFade::FWaveTransformationTrimFade(const FFadeFunctionData& InFadeFunctions, double InStartTime, double InEndTime, float InStartFadeTime, float InStartFadeCurve, float InStartSCurveSharpness, float InEndFadeTime, float InEndFadeCurve, float InEndSCurveSharpness)
+
+FWaveTransformationTrimFade::FWaveTransformationTrimFade(const FFadeFunctionData& InFadeFunctions, double InStartTime, double InEndTime)
 	: FadeFunctions(InFadeFunctions)
 	, StartTime(InStartTime)
-	, EndTime(InEndTime)
-	, StartFadeTime(InStartFadeTime)
-	, StartFadeCurve(FMath::Max(InStartFadeCurve, UWaveformTransformationTrimFade::MinFadeCurve))
-	, StartSCurveSharpness(InStartSCurveSharpness)
-	, EndFadeTime(InEndFadeTime)
-	, EndFadeCurve(FMath::Max(InEndFadeCurve, UWaveformTransformationTrimFade::MinFadeCurve))
-	, EndSCurveSharpness(InEndSCurveSharpness){}
-PRAGMA_ENABLE_DEPRECATION_WARNINGS
+	, EndTime(InEndTime){}
 
 void FWaveTransformationTrimFade::ProcessAudio(Audio::FWaveformTransformationWaveInfo& InOutWaveInfo) const
 {
@@ -125,28 +149,39 @@ void FWaveTransformationTrimFade::ProcessAudio(Audio::FWaveformTransformationWav
 
 	Audio::FAlignedFloatBuffer& InputAudio = *InOutWaveInfo.Audio;
 
-	if(InputAudio.Num() == 0)
+	// If InputAudio is smaller than the size of a frame, do not process the audio.
+	if (InputAudio.Num() < InOutWaveInfo.NumChannels)
 	{
 		return;
 	}
 
-	const int64 LastInputAudioIndex = InputAudio.Num() - 1;
-	int64 StartSample = FMath::Min(FMath::FloorToInt64(StartTime * InOutWaveInfo.SampleRate) * InOutWaveInfo.NumChannels, LastInputAudioIndex - InOutWaveInfo.NumChannels);
-	int64 EndSample = InputAudio.Num();
+	check(InOutWaveInfo.NumChannels > 0);
+
+	const int64 LastInputAudioIndex = static_cast<int64>(InputAudio.Num() - 1);
+	const int64 NumChannelsMinusOne = static_cast<int64>(InOutWaveInfo.NumChannels - 1); // Used to step backwards from the end sample to a valid frame
+	const int64 FirstSampleOfLastFrame = LastInputAudioIndex - NumChannelsMinusOne; // The last sample that begins a frame
+	int64 StartSample = FMath::Min(FMath::FloorToInt64(StartTime * InOutWaveInfo.SampleRate) * InOutWaveInfo.NumChannels, FirstSampleOfLastFrame);
+	int64 EndSample = LastInputAudioIndex;
 	
 	if(EndTime > 0.f)
 	{	
 		const int64 EndFrame = FMath::RoundToInt64(EndTime * InOutWaveInfo.SampleRate);
-		EndSample = EndFrame * InOutWaveInfo.NumChannels;
+		EndSample = EndFrame * InOutWaveInfo.NumChannels + NumChannelsMinusOne;
 		EndSample = FMath::Min(EndSample, LastInputAudioIndex);
 	}
 
-	if (StartSample > EndSample - InOutWaveInfo.NumChannels)
+	if (StartSample > EndSample - NumChannelsMinusOne)
 	{
-		StartSample = EndSample - InOutWaveInfo.NumChannels;
+		StartSample = FMath::Max(EndSample - NumChannelsMinusOne, 0);
 	}
 
-	const int64 FinalSize = EndSample - StartSample;
+	check(StartSample >= 0);
+	check(StartSample % InOutWaveInfo.NumChannels == 0);
+	check(EndSample >= StartSample);
+	check((EndSample - NumChannelsMinusOne) % InOutWaveInfo.NumChannels == 0);
+
+	const int64 FinalSize = (EndSample - StartSample) + 1;
+	check(FinalSize % InOutWaveInfo.NumChannels == 0);
 
 	InOutWaveInfo.StartFrameOffset = StartSample - (StartSample % InOutWaveInfo.NumChannels);
 	InOutWaveInfo.NumEditedSamples = FinalSize;
@@ -156,7 +191,9 @@ void FWaveTransformationTrimFade::ProcessAudio(Audio::FWaveformTransformationWav
 		return;
 	}
 
-	const bool bProcessFades = StartFadeTime > 0.f || EndFadeTime > 0.f;
+	const bool bProcessFadeIn = FadeFunctions.FadeIn && FadeFunctions.FadeIn->Duration > 0.f;
+	const bool bProcessFadeOut = FadeFunctions.FadeOut && FadeFunctions.FadeOut->Duration > 0.f;
+	const bool bProcessFades = bProcessFadeIn || bProcessFadeOut;
 
 	if (!bProcessFades && FinalSize == InputAudio.Num())
 	{
@@ -173,17 +210,15 @@ void FWaveTransformationTrimFade::ProcessAudio(Audio::FWaveformTransformationWav
 
 	FMemory::Memcpy(InputAudio.GetData(), TempBuffer.GetData(), FinalSize * sizeof(float));
 
-	PRAGMA_DISABLE_DEPRECATION_WARNINGS
-	if ((FadeFunctions.FadeIn && FadeFunctions.FadeIn->Duration > 0.f) || StartFadeTime > 0.f)
+	if (bProcessFadeIn)
 	{
-		ApplyFadeIn(InputAudio, FadeFunctions, StartFadeTime, StartFadeCurve, StartSCurveSharpness, InOutWaveInfo.NumChannels, InOutWaveInfo.SampleRate);
+		ApplyFadeIn(InputAudio, FadeFunctions, InOutWaveInfo.NumChannels, InOutWaveInfo.SampleRate);
 	}
 
-	if ((FadeFunctions.FadeOut && FadeFunctions.FadeOut->Duration > 0.f) || EndFadeTime > 0.f)
+	if (bProcessFadeOut)
 	{
-		ApplyFadeOut(InputAudio, FadeFunctions, EndFadeTime, EndFadeCurve, EndSCurveSharpness, InOutWaveInfo.NumChannels, InOutWaveInfo.SampleRate);
-	}	
-	PRAGMA_ENABLE_DEPRECATION_WARNINGS
+		ApplyFadeOut(InputAudio, FadeFunctions, InOutWaveInfo.NumChannels, InOutWaveInfo.SampleRate);
+	}
 }
 
 void UWaveformTransformationTrimFade::PostEditChangeProperty(FPropertyChangedEvent& PropertyChangedEvent)
@@ -192,11 +227,6 @@ void UWaveformTransformationTrimFade::PostEditChangeProperty(FPropertyChangedEve
 
 	const float MaxDuration = EndTime - StartTime;
 
-	PRAGMA_DISABLE_DEPRECATION_WARNINGS
-	EndFadeTime = FMath::Min(EndFadeTime, MaxDuration);
-	StartFadeTime = FMath::Min(StartFadeTime, MaxDuration);
-	PRAGMA_ENABLE_DEPRECATION_WARNINGS
-
 	if (FadeFunctions.FadeIn)
 	{
 		FadeFunctions.FadeIn->Duration = FMath::Min(FadeFunctions.FadeIn->Duration, MaxDuration);
@@ -210,9 +240,7 @@ void UWaveformTransformationTrimFade::PostEditChangeProperty(FPropertyChangedEve
 
 Audio::FTransformationPtr UWaveformTransformationTrimFade::CreateTransformation() const
 {
-	PRAGMA_DISABLE_DEPRECATION_WARNINGS
-	return MakeUnique<FWaveTransformationTrimFade>(FadeFunctions, StartTime, EndTime, StartFadeTime, StartFadeCurve, StartSCurveSharpness, EndFadeTime, EndFadeCurve, EndSCurveSharpness);
-	PRAGMA_ENABLE_DEPRECATION_WARNINGS
+	return MakeUnique<FWaveTransformationTrimFade>(FadeFunctions, StartTime, EndTime);
 }
 
 void UWaveformTransformationTrimFade::UpdateConfiguration(FWaveTransformUObjectConfiguration& InOutConfiguration)
@@ -223,6 +251,7 @@ void UWaveformTransformationTrimFade::UpdateConfiguration(FWaveTransformUObjectC
 	InOutConfiguration.EndTime = EndTime;
 }
 
+PRAGMA_DISABLE_DEPRECATION_WARNINGS
 const double UWaveformTransformationTrimFade::GetFadeInCurveValue(const float StartFadeCurve, const double FadeFraction, const float SCurveSharpness)
 {
 	double CurveValue = 1;
@@ -253,7 +282,7 @@ const double UWaveformTransformationTrimFade::GetFadeInCurveValue(const float St
 const double UWaveformTransformationTrimFade::GetFadeOutCurveValue(const float EndFadeCurve, const double FadeFraction, const float SCurveSharpness)
 {
 	double CurveValue = 1;
-	
+
 	if (EndFadeCurve < 0)
 	{
 		const double X = FadeFraction;
@@ -278,6 +307,103 @@ const double UWaveformTransformationTrimFade::GetFadeOutCurveValue(const float E
 
 	return CurveValue;
 }
+PRAGMA_ENABLE_DEPRECATION_WARNINGS
+
+void UWaveformTransformationTrimFade::PostInitProperties()
+{
+	Super::PostInitProperties();
+
+	if (FadeFunctions.FadeIn == nullptr)
+	{
+		PRAGMA_DISABLE_DEPRECATION_WARNINGS
+		if (StartFadeTime == 0.f)
+		PRAGMA_ENABLE_DEPRECATION_WARNINGS
+		{
+			FadeFunctions.FadeIn = NewObject<UFadeCurveFunctionExponential>(this, UFadeCurveFunctionExponential::StaticClass(), NAME_None, RF_Transactional);
+			FadeFunctions.FadeIn->Duration = 0.f;
+		}
+	}
+
+	check(FadeFunctions.FadeIn);
+
+	if (FadeFunctions.FadeOut == nullptr)
+	{
+		PRAGMA_DISABLE_DEPRECATION_WARNINGS
+		if (EndFadeTime == 0.f)
+		PRAGMA_ENABLE_DEPRECATION_WARNINGS
+		{
+			FadeFunctions.FadeOut = NewObject<UFadeCurveFunctionExponential>(this, UFadeCurveFunctionExponential::StaticClass(), NAME_None, RF_Transactional);
+			FadeFunctions.FadeOut->Duration = 0.f;
+		}
+	}
+
+	check(FadeFunctions.FadeOut);
+}
+
+void UWaveformTransformationTrimFade::PostLoad()
+{
+	Super::PostLoad();
+
+	const float Linear = 1.f;
+	
+	PRAGMA_DISABLE_DEPRECATION_WARNINGS
+	if (StartFadeTime != 0.f)
+	{
+		if (FadeFunctions.FadeIn == nullptr)
+		{
+			if (StartFadeCurve > Linear)
+			{
+				FadeFunctions.FadeIn = NewObject<UFadeCurveFunctionExponential>(this, UFadeCurveFunctionExponential::StaticClass(), NAME_None, RF_Transactional);
+			}
+			else if (StartFadeCurve < 0.f)
+			{
+				FadeFunctions.FadeIn = NewObject<UFadeCurveFunctionSigmoid>(this, UFadeCurveFunctionSigmoid::StaticClass(), NAME_None, RF_Transactional);
+			}
+			else if (StartFadeCurve < Linear)
+			{
+				FadeFunctions.FadeIn = NewObject<UFadeCurveFunctionLogarithmic>(this, UFadeCurveFunctionLogarithmic::StaticClass(), NAME_None, RF_Transactional);
+			}
+			else if (StartFadeCurve == Linear)
+			{
+				FadeFunctions.FadeIn = NewObject<UFadeFunctionLinear>(this, UFadeFunctionLinear::StaticClass(), NAME_None, RF_Transactional);
+			}
+			
+			FadeFunctions.FadeIn->Duration = StartFadeTime;
+			SetFadeCurve(FadeFunctions.FadeIn, StartFadeCurve, StartSCurveSharpness);
+		}
+
+		StartFadeTime = 0.f;
+	}
+
+	if (EndFadeTime != 0.f)
+	{
+		if (FadeFunctions.FadeOut == nullptr)
+		{
+			if (EndFadeCurve > Linear)
+			{
+				FadeFunctions.FadeOut = NewObject<UFadeCurveFunctionExponential>(this, UFadeCurveFunctionExponential::StaticClass(), NAME_None, RF_Transactional);
+			}
+			else if (EndFadeCurve < 0.f)
+			{
+				FadeFunctions.FadeOut = NewObject<UFadeCurveFunctionSigmoid>(this, UFadeCurveFunctionSigmoid::StaticClass(), NAME_None, RF_Transactional);
+			}
+			else if (EndFadeCurve < Linear)
+			{
+				FadeFunctions.FadeOut = NewObject<UFadeCurveFunctionLogarithmic>(this, UFadeCurveFunctionLogarithmic::StaticClass(), NAME_None, RF_Transactional);
+			}
+			else if (EndFadeCurve == Linear)
+			{
+				FadeFunctions.FadeOut = NewObject<UFadeFunctionLinear>(this, UFadeFunctionLinear::StaticClass(), NAME_None, RF_Transactional);
+			}
+			
+			FadeFunctions.FadeOut->Duration = EndFadeTime;
+			SetFadeCurve(FadeFunctions.FadeOut, EndFadeCurve, EndSCurveSharpness);
+		}
+
+		EndFadeTime = 0.f;
+	}
+	PRAGMA_ENABLE_DEPRECATION_WARNINGS
+}
 
 void UWaveformTransformationTrimFade::UpdateDurationProperties(const float InAvailableDuration)
 {
@@ -394,3 +520,18 @@ const double UFadeCurveFunctionSigmoid::GetFadeOutCurveValue(const double FadeFr
 
 	return CurveValue;
 }
+
+void UFadeFunction::PostEditUndo()
+{
+	Super::PostEditUndo();
+
+	UObject* Outer = GetOuter();
+	if (Outer && Outer->IsA<UWaveformTransformationBase>())
+	{
+		UWaveformTransformationBase* Transformation = CastChecked<UWaveformTransformationBase>(GetOuter());
+
+		static const bool bMarkFileDirty = true;
+
+		Transformation->OnTransformationChanged.ExecuteIfBound(bMarkFileDirty);
+	}
+}
diff --git a/Engine/Plugins/Editor/WaveformEditor/Source/WaveformTransformations/Public/WaveformTransformationTrimFade.h b/Engine/Plugins/Editor/WaveformEditor/Source/WaveformTransformations/Public/WaveformTransformationTrimFade.h
--- a/Engine/Plugins/Editor/WaveformEditor/Source/WaveformTransformations/Public/WaveformTransformationTrimFade.h
+++ b/Engine/Plugins/Editor/WaveformEditor/Source/WaveformTransformations/Public/WaveformTransformationTrimFade.h
@@ -20,6 +20,11 @@ public:
 
 	UPROPERTY(EditAnywhere, Category = "Fade", BlueprintReadWrite, meta = (ClampMin = 0.0, DisplayName = "Fade Duration"))
 	float Duration = 0.5f;
+
+private: 
+#if WITH_EDITOR
+	virtual void PostEditUndo() override;
+#endif // WITH_EDITOR
 };
 
 UCLASS(MinimalAPI, Abstract, EditInlineNew, DefaultToInstanced, CollapseCategories, NotBlueprintable)
@@ -129,7 +134,7 @@ enum class EWaveEditorFadeMode : uint8
 class FWaveTransformationTrimFade : public Audio::IWaveTransformation
 {
 public:
-	UE_API explicit FWaveTransformationTrimFade(const FFadeFunctionData& InFadeFunctions, double InStartTime, double InEndTime, float InStartFadeTime, float InStartFadeCurve, float InStartSCurveSharpness, float InEndFadeTime, float InEndFadeCurve, float InEndSCurveSharpness);
+	UE_API explicit FWaveTransformationTrimFade(const FFadeFunctionData& InFadeFunctions, double InStartTime, double InEndTime);
 	UE_API virtual void ProcessAudio(Audio::FWaveformTransformationWaveInfo& InOutWaveInfo) const override;
 
 	virtual constexpr Audio::ETransformationPriority FileChangeLengthPriority() const override { return Audio::ETransformationPriority::High; }
@@ -140,14 +145,6 @@ private:
 
 	double StartTime = 0.0;
 	double EndTime = 0.0;
-
-	float StartFadeTime = 0.f;
-	float StartFadeCurve = 0.f;
-	float StartSCurveSharpness = 0.f;
-
-	float EndFadeTime = 0.f;
-	float EndFadeCurve = 0.f;
-	float EndSCurveSharpness = 0.f;
 };
 
 UCLASS(MinimalAPI)
@@ -174,21 +171,21 @@ public:
 	 * Relative to StartTime
 	 * Individual fade properties have been deprecated. Use FadeFunctions instead.
 	 */ 
-	UPROPERTY(EditAnywhere, Category = "Fade", meta = (ClampMin = 0.0, DisplayName = "Fade-In Duration"))
+	UPROPERTY(EditAnywhere, Category = "Fade", meta = (ClampMin = 0.0, DisplayName = "Fade-In Duration", EditCondition = "StartFadeTime != 0", EditConditionHides))
 	float StartFadeTime = 0.f;
 	
 	UE_DEPRECATED(5.7, "Individual fade properties have been deprecated. Use FadeFunctions instead.")
 	/*
 	 * Individual fade properties have been deprecated. Use FadeFunctions instead.
 	 */
-	UPROPERTY(EditAnywhere, Category = "Fade", meta=(ClampMin = -0.1, ClampMax = 10.0, DisplayName = "Fade-In Curve"))
+	UPROPERTY(EditAnywhere, Category = "Fade", meta=(ClampMin = -0.1, ClampMax = 10.0, DisplayName = "Fade-In Curve", EditCondition = "StartFadeTime != 0", EditConditionHides))
 	float StartFadeCurve = 1.f;
 
 	UE_DEPRECATED(5.7, "Individual fade properties have been deprecated. Use FadeFunctions instead.")
 	/*
 	 * Individual fade properties have been deprecated. Use FadeFunctions instead.
 	 */
-	UPROPERTY(EditAnywhere, Category = "Fade", meta = (ClampMin = 0, ClampMax = 1.0, DisplayName = "S-Curve Sharpness", EditCondition = "StartFadeCurve < 0", EditConditionHides))
+	UPROPERTY(EditAnywhere, Category = "Fade", meta = (ClampMin = 0, ClampMax = 1.0, DisplayName = "S-Curve Sharpness", EditCondition = "StartFadeTime != 0 && StartFadeCurve < 0", EditConditionHides))
 	float StartSCurveSharpness = 0.75f;
 
 	UE_DEPRECATED(5.7, "Individual fade properties have been deprecated. Use FadeFunctions instead.")
@@ -196,21 +193,21 @@ public:
 	 * Relative to EndTime. Fade out starts at EndTime - EndFadeTime and ends at EndTime
 	 * Individual fade properties have been deprecated. Use FadeFunctions instead.
 	 */
-	UPROPERTY(EditAnywhere, Category = "Fade", meta = (ClampMin = 0.0, DisplayName = "Fade-Out Duration"))
+	UPROPERTY(EditAnywhere, Category = "Fade", meta = (ClampMin = 0.0, DisplayName = "Fade-Out Duration", EditCondition = "EndFadeTime != 0", EditConditionHides))
 	float EndFadeTime = 0.f;
 
 	UE_DEPRECATED(5.7, "Individual fade properties have been deprecated. Use FadeFunctions instead.")
 	/*
 	 * Individual fade properties have been deprecated. Use FadeFunctions instead.
 	 */
-	UPROPERTY(EditAnywhere, Category = "Fade", meta=(ClampMin = -0.1, ClampMax = 10.0, DisplayName = "Fade-Out Curve"))
+	UPROPERTY(EditAnywhere, Category = "Fade", meta=(ClampMin = -0.1, ClampMax = 10.0, DisplayName = "Fade-Out Curve", EditCondition = "EndFadeTime != 0", EditConditionHides))
 	float EndFadeCurve = 1.f;
 
 	UE_DEPRECATED(5.7, "Individual fade properties have been deprecated. Use FadeFunctions instead.")
 	/*
 	 * Individual fade properties have been deprecated. Use FadeFunctions instead.
 	 */
-	UPROPERTY(EditAnywhere, Category = "Fade", meta = (ClampMin = 0, ClampMax = 1.0, DisplayName = "S-Curve Sharpness", EditCondition = "EndFadeCurve < 0", EditConditionHides))
+	UPROPERTY(EditAnywhere, Category = "Fade", meta = (ClampMin = 0, ClampMax = 1.0, DisplayName = "S-Curve Sharpness", EditCondition = "EndFadeTime != 0 && EndFadeCurve < 0", EditConditionHides))
 	float EndSCurveSharpness = 0.75;
 
 	UE_DEPRECATED(5.7, "Individual fade properties have been deprecated. Use FadeFunctions instead.")
@@ -219,13 +216,15 @@ public:
 	UE_DEPRECATED(5.7, "Individual fade properties have been deprecated. Use FadeFunctions instead.")
 	static UE_API const float MaxFadeCurve;
 	
+	UE_DEPRECATED(5.7, "Individual fade properties have been deprecated. Use FadeFunctions instead.")
 	static UE_API const TMap<EWaveEditorFadeMode, float> FadeModeToCurveValueMap;
+
 	static UE_API const TMap<EWaveEditorFadeMode, TSubclassOf<UFadeFunction>> FadeModeToFadeFunctionMap;
 
 	UE_API virtual Audio::FTransformationPtr CreateTransformation() const override;
 
 	UE_API void UpdateConfiguration(FWaveTransformUObjectConfiguration& InOutConfiguration) override;
-	virtual constexpr Audio::ETransformationPriority GetTransformationPriority() const { return Audio::ETransformationPriority::High; }
+	virtual constexpr Audio::ETransformationPriority GetTransformationPriority() const { return Audio::ETransformationPriority::High; }	
 
 	UE_DEPRECATED(5.7, "Individual fade properties have been deprecated. Use FadeFunctions::FadeIn::GetFadeInCurveValue instead.")
 	static UE_API const double GetFadeInCurveValue(const float StartFadeCurve, const double FadeFraction, const float SCurveSharpness = 0);
@@ -234,6 +233,10 @@ public:
 	static UE_API const double GetFadeOutCurveValue(const float EndFadeCurve, const double FadeFraction, const float SCurveSharpness = 0);
 	
 private:	
+
+	UE_API virtual void PostInitProperties() override;
+	virtual void PostLoad() override;
+
 	UE_API void UpdateDurationProperties(const float InAvailableDuration);
 	float AvailableWaveformDuration = -1.f;
 };
diff --git a/Engine/Plugins/Editor/WaveformEditor/Source/WaveformTransformationsWidgets/Private/WaveformTransformationTrimFadeRenderer.cpp b/Engine/Plugins/Editor/WaveformEditor/Source/WaveformTransformationsWidgets/Private/WaveformTransformationTrimFadeRenderer.cpp
--- a/Engine/Plugins/Editor/WaveformEditor/Source/WaveformTransformationsWidgets/Private/WaveformTransformationTrimFadeRenderer.cpp
+++ b/Engine/Plugins/Editor/WaveformEditor/Source/WaveformTransformationsWidgets/Private/WaveformTransformationTrimFadeRenderer.cpp
@@ -158,54 +158,76 @@ int32 FWaveformTransformationTrimFadeRenderer::DrawFadeCurves(const FGeometry& A
 	return LayerId;
 }
 
-PRAGMA_DISABLE_DEPRECATION_WARNINGS
 void FWaveformTransformationTrimFadeRenderer::GenerateFadeCurves(const FGeometry& AllottedGeometry)
 {
 	check(StrongTrimFade);
-	const float FadeInTime = StrongTrimFade->FadeFunctions.FadeIn ? StrongTrimFade->FadeFunctions.FadeIn->Duration : StrongTrimFade->StartFadeTime;
-	const float FadeInFrames = FadeInTime * TransformationWaveInfo.SampleRate;
-	const uint32 FadeInPixelLength = FadeInFrames * PixelsPerFrame;
+
+	check(StartTimeHandleX >= 0);
+	check(EndTimeHandleX >= 0);
 	FadeInStartX = FMath::RoundToInt32(StartTimeHandleX);
-	FadeInEndX = FMath::RoundToInt32(FMath::Clamp(StartTimeHandleX + FadeInPixelLength, StartTimeHandleX, EndTimeHandleX));
-	
-	const uint32 DisplayedFadeInPixelLength = FadeInEndX - FadeInStartX;
-	FadeInCurvePoints.SetNumUninitialized(DisplayedFadeInPixelLength);
+	FadeOutEndX = FMath::RoundToInt32(EndTimeHandleX);
 
-	for (uint32 Pixel = 0; Pixel < DisplayedFadeInPixelLength; ++Pixel)
+	if (StrongTrimFade->FadeFunctions.FadeIn && StrongTrimFade->FadeFunctions.FadeIn->Duration > 0.f)
 	{
-		const double FadeFraction = (float)Pixel / FadeInPixelLength;
-		const double CurveFunction = StrongTrimFade->FadeFunctions.FadeIn ? StrongTrimFade->FadeFunctions.FadeIn->GetFadeInCurveValue(FadeFraction)
-			: UWaveformTransformationTrimFade::GetFadeInCurveValue(StrongTrimFade->StartFadeCurve, FadeFraction, StrongTrimFade->StartSCurveSharpness);
+		const float FadeInTime = StrongTrimFade->FadeFunctions.FadeIn->Duration;
+		const float FadeInFrames = FadeInTime * TransformationWaveInfo.SampleRate;
+		const float FadeInPixelLength = FadeInFrames * PixelsPerFrame;
+		check(FadeInPixelLength > 0.f);
+		check(EndTimeHandleX >= 0);
+		
+		FadeInEndX = FMath::RoundToInt32(FMath::Clamp(StartTimeHandleX + FadeInPixelLength, StartTimeHandleX, EndTimeHandleX));
 
-		const double CurveValue = Pixel != FadeInPixelLength - 1 ? 1.f - CurveFunction : 0.f;
+		const uint32 DisplayedFadeInPixelLength = FadeInEndX - FadeInStartX;
+		FadeInCurvePoints.SetNumUninitialized(DisplayedFadeInPixelLength);
 
-		const uint32 XCoordinate = Pixel + FadeInStartX;
-		FadeInCurvePoints[Pixel] = FVector2D(XCoordinate, CurveValue * AllottedGeometry.Size.Y);
-	}
+		for (uint32 Pixel = 0; Pixel < DisplayedFadeInPixelLength; ++Pixel)
+		{
+			const double FadeFraction = static_cast<float>(Pixel) / FadeInPixelLength;
+			const double CurveFunction = StrongTrimFade->FadeFunctions.FadeIn->GetFadeInCurveValue(FadeFraction);
 
-	const float FadeOutTime = StrongTrimFade->FadeFunctions.FadeOut ? StrongTrimFade->FadeFunctions.FadeOut->Duration : StrongTrimFade->EndFadeTime;
-	const float FadeOutFrames = FadeOutTime * TransformationWaveInfo.SampleRate;
-	const float FadeOutPixelLength = FadeOutFrames * PixelsPerFrame;
-	FadeOutStartX = FMath::RoundToInt32(FMath::Clamp(EndTimeHandleX - FadeOutPixelLength, StartTimeHandleX, EndTimeHandleX));
-	FadeOutEndX = FMath::RoundToInt32(EndTimeHandleX);
-	
-	const uint32 DisplayedFadeOutPixelLength = FadeOutEndX - FadeOutStartX;
-	FadeOutCurvePoints.SetNumUninitialized(DisplayedFadeOutPixelLength);
-	const uint32 FadeOutPixelOffset = FadeOutPixelLength - DisplayedFadeOutPixelLength;
+			const double CurveValue = Pixel != FadeInPixelLength - 1 ? 1.f - CurveFunction : 0.f;
+
+			const uint32 XCoordinate = Pixel + FadeInStartX;
+			FadeInCurvePoints[Pixel] = FVector2D(XCoordinate, CurveValue * AllottedGeometry.Size.Y);
+		}
+	}
+	else
+	{
+		FadeInEndX = FMath::RoundToInt32(FMath::Clamp(StartTimeHandleX, StartTimeHandleX, EndTimeHandleX));
+		FadeInCurvePoints.SetNumUninitialized(0);
+	}
 
-	for (uint32 Pixel = 0; Pixel < DisplayedFadeOutPixelLength; ++Pixel)
+	if(StrongTrimFade->FadeFunctions.FadeOut && StrongTrimFade->FadeFunctions.FadeOut->Duration > 0.f)
 	{
-		const double FadeFraction = (float)(Pixel + FadeOutPixelOffset) / FadeOutPixelLength;
-		const double CurveFunction = StrongTrimFade->FadeFunctions.FadeOut ? StrongTrimFade->FadeFunctions.FadeOut->GetFadeOutCurveValue(FadeFraction)
-			: UWaveformTransformationTrimFade::GetFadeOutCurveValue(StrongTrimFade->EndFadeCurve, FadeFraction, StrongTrimFade->EndSCurveSharpness);
+		const float FadeOutTime = StrongTrimFade->FadeFunctions.FadeOut->Duration;
+		const float FadeOutFrames = FadeOutTime * TransformationWaveInfo.SampleRate;
+		const float FadeOutPixelLength = FadeOutFrames * PixelsPerFrame;
+		check(FadeOutPixelLength > 0.f);
 
-		const double CurveValue = Pixel != FadeOutPixelLength - 1 ? 1.f - CurveFunction : 1.f;
+		check(StartTimeHandleX >= 0);
+		FadeOutStartX = FMath::RoundToInt32(FMath::Clamp(EndTimeHandleX - FadeOutPixelLength, StartTimeHandleX, EndTimeHandleX));
+
+		const uint32 DisplayedFadeOutPixelLength = FadeOutEndX - FadeOutStartX;
+		FadeOutCurvePoints.SetNumUninitialized(DisplayedFadeOutPixelLength);
+		const uint32 FadeOutPixelOffset = FMath::RoundToInt32(FMath::Max(FadeOutPixelLength - static_cast<float>(DisplayedFadeOutPixelLength), 0.f));
+
+		for (uint32 Pixel = 0; Pixel < DisplayedFadeOutPixelLength; ++Pixel)
+		{
+			const double FadeFraction = static_cast<float>(Pixel + FadeOutPixelOffset) / FadeOutPixelLength;
+			const double CurveFunction = StrongTrimFade->FadeFunctions.FadeOut->GetFadeOutCurveValue(FadeFraction);
 
-		const uint32 XCoordinate = Pixel + FadeOutStartX;
-		FadeOutCurvePoints[Pixel] = FVector2D(XCoordinate, CurveValue * AllottedGeometry.Size.Y);
+			const double CurveValue = Pixel != FadeOutPixelLength - 1 ? 1.f - CurveFunction : 1.f;
+
+			const uint32 XCoordinate = Pixel + FadeOutStartX;
+			FadeOutCurvePoints[Pixel] = FVector2D(XCoordinate, CurveValue * AllottedGeometry.Size.Y);
+		}
+	}
+	else
+	{
+		FadeOutStartX = FMath::RoundToInt32(FMath::Clamp(EndTimeHandleX, StartTimeHandleX, EndTimeHandleX));
+		FadeOutCurvePoints.SetNumUninitialized(0);
 	}
 }
-PRAGMA_ENABLE_DEPRECATION_WARNINGS
 
 FCursorReply FWaveformTransformationTrimFadeRenderer::OnCursorQuery(const FGeometry& MyGeometry, const FPointerEvent& CursorEvent) const
 {
@@ -242,46 +264,24 @@ FReply FWaveformTransformationTrimFadeRenderer::OnMouseWheel(SWidget& OwnerWidge
 	check(StrongTrimFade);
 	const FVector2D LocalCursorPosition = GetLocalCursorPosition(MouseEvent, MyGeometry);
 
-	if (IsCursorInFadeInInteractionRange(LocalCursorPosition, MyGeometry))
-	{		
+	if (IsCursorInFadeInInteractionRange(LocalCursorPosition, MyGeometry) && StrongTrimFade->FadeFunctions.FadeIn && StrongTrimFade->FadeFunctions.FadeIn.IsA<UFadeCurveFunction>())
+	{
 		const float FadeCurveDelta = MouseEvent.GetWheelDelta() * MouseWheelStep;
 		FText OutPropertyName = FText::FromName(NAME_None);
-		
-		if (StrongTrimFade->FadeFunctions.FadeIn && StrongTrimFade->FadeFunctions.FadeIn.IsA<UFadeCurveFunction>())
-		{
-			TObjectPtr<UFadeCurveFunction> FadeIn = CastChecked<UFadeCurveFunction>(StrongTrimFade->FadeFunctions.FadeIn);
-			OutPropertyName = FadeIn->GetFadeCurvePropertyName();
-			
-			if (BeginTransaction(TEXT("PropertyEditor"), FText::Format(NSLOCTEXT("PropertyEditor", "SetNumericPropertyTransaction", "Edit {0}"), OutPropertyName), StrongTrimFade.Get()) == INDEX_NONE)
-			{
-				UE_LOG(LogWaveformTransformation, Warning, TEXT("Begin TrimFadeTransformation OnMouseWheel Transaction Failed"));
-			}
 
-			StrongTrimFade->Modify();
-			StrongTrimFade->FadeFunctions.FadeIn->Modify();
+		TObjectPtr<UFadeCurveFunction> FadeIn = CastChecked<UFadeCurveFunction>(StrongTrimFade->FadeFunctions.FadeIn);
+		OutPropertyName = FadeIn->GetFadeCurvePropertyName();
 
-			const float FadeCurveValue = FadeIn->GetFadeCurve() + FadeCurveDelta;
-			FadeIn->SetFadeCurve(FadeCurveValue);
-		}
-		else
+		if (BeginTransaction(TEXT("PropertyEditor"), FText::Format(NSLOCTEXT("PropertyEditor", "SetNumericPropertyTransaction", "Edit {0}"), OutPropertyName), StrongTrimFade.Get()) == INDEX_NONE)
 		{
-			PRAGMA_DISABLE_DEPRECATION_WARNINGS
-			OutPropertyName = FText::FromName(GET_MEMBER_NAME_CHECKED(UWaveformTransformationTrimFade, StartFadeCurve));
-			if (BeginTransaction(TEXT("PropertyEditor"), FText::Format(NSLOCTEXT("PropertyEditor", "SetNumericPropertyTransaction", "Edit {0}"), OutPropertyName), StrongTrimFade.Get()) == INDEX_NONE)
-			{
-				UE_LOG(LogWaveformTransformation, Warning, TEXT("Begin TrimFadeTransformation OnMouseWheel Transaction Failed"));
-			}
-
-			StrongTrimFade->Modify();
-
-			const float FadeCurveValue = FMath::Clamp(StrongTrimFade->StartFadeCurve + FadeCurveDelta, 0.f, 10.f);
-			StrongTrimFade->StartFadeCurve = FadeCurveValue;
-			PRAGMA_ENABLE_DEPRECATION_WARNINGS
+			UE_LOG(LogWaveformTransformation, Warning, TEXT("Begin TrimFadeTransformation OnMouseWheel Transaction Failed"));
 		}
 
-		EPropertyValueSetFlags::Type PropertyFlags = EPropertyValueSetFlags::InteractiveChange;
+		StrongTrimFade->Modify();
+		StrongTrimFade->FadeFunctions.FadeIn->Modify();
 
-		StrongTrimFade->PostEditChange();
+		const float FadeCurveValue = FadeIn->GetFadeCurve() + FadeCurveDelta;
+		FadeIn->SetFadeCurve(FadeCurveValue);
 
 		EndTransaction();
 		StrongTrimFade->OnTransformationChanged.ExecuteIfBound(true);
@@ -289,46 +289,24 @@ FReply FWaveformTransformationTrimFadeRenderer::OnMouseWheel(SWidget& OwnerWidge
 		return FReply::Handled();
 	}
 
-	if (IsCursorInFadeOutInteractionRange(LocalCursorPosition, MyGeometry))
+	if (IsCursorInFadeOutInteractionRange(LocalCursorPosition, MyGeometry) && StrongTrimFade->FadeFunctions.FadeOut && StrongTrimFade->FadeFunctions.FadeOut.IsA<UFadeCurveFunction>())
 	{
 		const float FadeCurveDelta = MouseEvent.GetWheelDelta() * MouseWheelStep;
 		FText OutPropertyName = FText::FromName(NAME_None);
 
-		if (StrongTrimFade->FadeFunctions.FadeOut && StrongTrimFade->FadeFunctions.FadeOut.IsA<UFadeCurveFunction>())
+		TObjectPtr<UFadeCurveFunction> FadeOut = CastChecked<UFadeCurveFunction>(StrongTrimFade->FadeFunctions.FadeOut);
+		OutPropertyName = FadeOut->GetFadeCurvePropertyName();
+		if (BeginTransaction(TEXT("PropertyEditor"), FText::Format(NSLOCTEXT("PropertyEditor", "SetNumericPropertyTransaction", "Edit {0}"), OutPropertyName), StrongTrimFade.Get()) == INDEX_NONE)
 		{
-			TObjectPtr<UFadeCurveFunction> FadeOut = CastChecked<UFadeCurveFunction>(StrongTrimFade->FadeFunctions.FadeOut);
-			OutPropertyName = FadeOut->GetFadeCurvePropertyName();
-			if (BeginTransaction(TEXT("PropertyEditor"), FText::Format(NSLOCTEXT("PropertyEditor", "SetNumericPropertyTransaction", "Edit {0}"), OutPropertyName), StrongTrimFade.Get()) == INDEX_NONE)
-			{
-				UE_LOG(LogWaveformTransformation, Warning, TEXT("Begin TrimFadeTransformation OnMouseWheel Transaction Failed"));
-			}
-
-			StrongTrimFade->Modify();
-			StrongTrimFade->FadeFunctions.FadeOut->Modify();
-
-			const float FadeCurveValue = FadeOut->GetFadeCurve() + FadeCurveDelta;
-			FadeOut->SetFadeCurve(FadeCurveValue);
+			UE_LOG(LogWaveformTransformation, Warning, TEXT("Begin TrimFadeTransformation OnMouseWheel Transaction Failed"));
 		}
-		else
-		{
-			PRAGMA_DISABLE_DEPRECATION_WARNINGS
-			OutPropertyName = FText::FromName(GET_MEMBER_NAME_CHECKED(UWaveformTransformationTrimFade, EndFadeCurve));
-			if (BeginTransaction(TEXT("PropertyEditor"), FText::Format(NSLOCTEXT("PropertyEditor", "SetNumericPropertyTransaction", "Edit {0}"), OutPropertyName), StrongTrimFade.Get()) == INDEX_NONE)
-			{
-				UE_LOG(LogWaveformTransformation, Warning, TEXT("Begin TrimFadeTransformation OnMouseWheel Transaction Failed"));
-			}
 
-			StrongTrimFade->Modify();
+		StrongTrimFade->Modify();
+		StrongTrimFade->FadeFunctions.FadeOut->Modify();
 
-			const float FadeCurveValue = FMath::Clamp(StrongTrimFade->EndFadeCurve + FadeCurveDelta, 0.f, 10.f);
-			StrongTrimFade->EndFadeCurve = FadeCurveValue;
-			PRAGMA_ENABLE_DEPRECATION_WARNINGS
-		}
-		
-		EPropertyValueSetFlags::Type PropertyFlags = EPropertyValueSetFlags::InteractiveChange;
+		const float FadeCurveValue = FadeOut->GetFadeCurve() + FadeCurveDelta;
+		FadeOut->SetFadeCurve(FadeCurveValue);
 
-		StrongTrimFade->PostEditChange();
-		
 		EndTransaction();
 		StrongTrimFade->OnTransformationChanged.ExecuteIfBound(true);
 
@@ -408,7 +386,18 @@ FReply FWaveformTransformationTrimFadeRenderer::OnMouseButtonDown(SWidget& Owner
 		{
 			UE_LOG(LogWaveformTransformation, Warning, TEXT("Begin TrimFadeTransformation OnMouseButtonDown Transaction Failed"));
 		}
+
+		if (TrimFadeInteractionType == ETrimFadeInteractionType::ScrubbingFadeIn && StrongTrimFade->FadeFunctions.FadeIn)
+		{
+			StrongTrimFade->FadeFunctions.FadeIn->Modify();
+		}
+		else if (TrimFadeInteractionType == ETrimFadeInteractionType::ScrubbingFadeOut && StrongTrimFade->FadeFunctions.FadeOut)
+		{
+			StrongTrimFade->FadeFunctions.FadeOut->Modify();
+		}
+		
 		StrongTrimFade->Modify();
+
 		return FReply::Handled().CaptureMouse(OwnerWidget.AsShared()).PreventThrottling();
 	}
 
@@ -422,6 +411,9 @@ FReply FWaveformTransformationTrimFadeRenderer::OnMouseMove(SWidget& OwnerWidget
 	if (MouseEvent.IsMouseButtonDown(EKeys::LeftMouseButton) && TrimFadeInteractionType != ETrimFadeInteractionType::None)
 	{
 		SetPropertyValueDependingOnInteractionType(MouseEvent, MyGeometry);
+
+		StrongTrimFade->PostEditChange();
+
 		StrongTrimFade->OnTransformationRenderChanged.ExecuteIfBound();
 
 		return FReply::Handled().CaptureMouse(OwnerWidget.AsShared());
@@ -435,11 +427,12 @@ FReply FWaveformTransformationTrimFadeRenderer::OnMouseButtonUp(SWidget& OwnerWi
 	check(StrongTrimFade);
 	if (TrimFadeInteractionType != ETrimFadeInteractionType::None)
 	{
-		StrongTrimFade->Modify();
-
 		SetPropertyValueDependingOnInteractionType(MouseEvent, MyGeometry);
 
 		EndTransaction();
+
+		StrongTrimFade->PostEditChange();
+
 		StrongTrimFade->OnTransformationChanged.ExecuteIfBound(true);
 
 		TrimFadeInteractionType = ETrimFadeInteractionType::None;
@@ -470,39 +463,21 @@ void FWaveformTransformationTrimFadeRenderer::SetPropertyValueDependingOnInterac
 	case FWaveformTransformationTrimFadeRenderer::ETrimFadeInteractionType::ScrubbingFadeIn:
 		{
 			float StartFadeTimeValue = 0;
-
 			if (StrongTrimFade->FadeFunctions.FadeIn)
 			{
-				StrongTrimFade->FadeFunctions.FadeIn->Modify();
 				StartFadeTimeValue = FMath::Clamp(SelectedTime - StrongTrimFade->StartTime, 0.f, TNumericLimits<float>().Max());
 				StrongTrimFade->FadeFunctions.FadeIn->Duration = StartFadeTimeValue;
 			}
-			else
-			{
-				PRAGMA_DISABLE_DEPRECATION_WARNINGS
-				StartFadeTimeValue = FMath::Clamp(SelectedTime - StrongTrimFade->StartTime, 0.f, TNumericLimits<float>().Max());
-				StrongTrimFade->StartFadeTime = StartFadeTimeValue;
-				PRAGMA_ENABLE_DEPRECATION_WARNINGS
-			}
 		}
 		break;
 	case FWaveformTransformationTrimFadeRenderer::ETrimFadeInteractionType::ScrubbingFadeOut:
 		{	
 			float EndFadeTimeValue = 0;
-
 			if (StrongTrimFade->FadeFunctions.FadeOut)
 			{
-				StrongTrimFade->FadeFunctions.FadeOut->Modify();
 				EndFadeTimeValue = FMath::Clamp(StrongTrimFade->EndTime - SelectedTime, 0.f, TNumericLimits<float>().Max());
 				StrongTrimFade->FadeFunctions.FadeOut->Duration = EndFadeTimeValue;
 			}
-			else
-			{
-				PRAGMA_DISABLE_DEPRECATION_WARNINGS
-				EndFadeTimeValue = FMath::Clamp(StrongTrimFade->EndTime - SelectedTime, 0.f, TNumericLimits<float>().Max());
-				StrongTrimFade->EndFadeTime = EndFadeTimeValue;
-				PRAGMA_ENABLE_DEPRECATION_WARNINGS
-			}
 		}
 		break;
 	case FWaveformTransformationTrimFadeRenderer::ETrimFadeInteractionType::RightClickFadeIn:
@@ -544,7 +519,7 @@ void FWaveformTransformationTrimFadeRenderer::ShowSelectFadeModeMenuAtCursor(con
 	FadeModeMenuWindow = MenuWindow;
 	TWeakObjectPtr<UWaveformTransformationTrimFade> InWeakTrimFade = StrongTrimFade.Get();
 
-	for (const TPair<EWaveEditorFadeMode, float>& FadeOptionPair : UWaveformTransformationTrimFade::FadeModeToCurveValueMap)
+	for (const TPair<EWaveEditorFadeMode, TSubclassOf<UFadeFunction>>& FadeOptionPair : UWaveformTransformationTrimFade::FadeModeToFadeFunctionMap)
 	{
 		MenuContent->AddSlot()
 			.Padding(5)
@@ -562,24 +537,24 @@ void FWaveformTransformationTrimFadeRenderer::ShowSelectFadeModeMenuAtCursor(con
 
 							PinnedTrimFade->Modify();
 
-							PRAGMA_DISABLE_DEPRECATION_WARNINGS
 							if (InteractionType == FWaveformTransformationTrimFadeRenderer::ETrimFadeInteractionType::RightClickFadeIn)
 							{
 								if (PinnedTrimFade->FadeFunctions.FadeIn)
 								{
-									PinnedTrimFade->FadeFunctions.FadeIn = NewObject<UFadeFunction>(PinnedTrimFade.Get(), UWaveformTransformationTrimFade::FadeModeToFadeFunctionMap[FadeOptionPair.Key].Get());
+									const float Duration = PinnedTrimFade->FadeFunctions.FadeIn->Duration;
+									PinnedTrimFade->FadeFunctions.FadeIn = NewObject<UFadeFunction>(PinnedTrimFade.Get(), FadeOptionPair.Value.Get(), NAME_None, RF_Transactional);
+									PinnedTrimFade->FadeFunctions.FadeIn->Duration =  Duration;
 								}
-								PinnedTrimFade->StartFadeCurve = FadeOptionPair.Value;
 							}
 							else if (InteractionType == FWaveformTransformationTrimFadeRenderer::ETrimFadeInteractionType::RightClickFadeOut)
 							{
 								if (PinnedTrimFade->FadeFunctions.FadeOut)
 								{
-									PinnedTrimFade->FadeFunctions.FadeOut = NewObject<UFadeFunction>(PinnedTrimFade.Get(), UWaveformTransformationTrimFade::FadeModeToFadeFunctionMap[FadeOptionPair.Key].Get());
+									const float Duration = PinnedTrimFade->FadeFunctions.FadeOut->Duration;
+									PinnedTrimFade->FadeFunctions.FadeOut = NewObject<UFadeFunction>(PinnedTrimFade.Get(), FadeOptionPair.Value.Get(), NAME_None, RF_Transactional);
+									PinnedTrimFade->FadeFunctions.FadeOut->Duration = Duration;
 								}
-								PinnedTrimFade->EndFadeCurve = FadeOptionPair.Value;
 							}
-							PRAGMA_ENABLE_DEPRECATION_WARNINGS
 
 							if (GEditor && GEditor->Trans)
 							{
@@ -640,20 +615,29 @@ void FWaveformTransformationTrimFadeRenderer::ShowSelectFadeModeMenuAtCursor(con
 			return false; // return false to only execute once
 		}));
 
-	ApplicationActivationStateHandle = FSlateApplication::Get().OnApplicationActivationStateChanged().AddLambda([InFadeModeMenuWindow = FadeModeMenuWindow, InApplicationActivationStateHandle = ApplicationActivationStateHandle](bool isActive)
+	TSharedPtr<FDelegateHandle> ApplicationActivationStateHandlePtr = MakeShared<FDelegateHandle>();
+	*ApplicationActivationStateHandlePtr = FSlateApplication::Get().OnApplicationActivationStateChanged().AddLambda([InFadeModeMenuWindow = FadeModeMenuWindow, ApplicationActivationStateHandlePtr](bool isActive)
 		{
 			if (!isActive)
 			{
 				if (TSharedPtr<SWindow> LockedFadeModeMenuWindow = InFadeModeMenuWindow.Pin())
 				{
-					FSlateApplication::Get().RequestDestroyWindow(LockedFadeModeMenuWindow.ToSharedRef());	
-					FSlateApplication::Get().OnApplicationActivationStateChanged().Remove(InApplicationActivationStateHandle);
+					FSlateApplication::Get().RequestDestroyWindow(LockedFadeModeMenuWindow.ToSharedRef());
+					FSlateApplication::Get().OnApplicationActivationStateChanged().Remove(*ApplicationActivationStateHandlePtr);
 				}
 			}
 		});
 
+	if (ApplicationActivationStateHandle.IsValid())
+	{
+		FSlateApplication::Get().OnApplicationActivationStateChanged().Remove(ApplicationActivationStateHandle);
+	}
+
+	ApplicationActivationStateHandle = *ApplicationActivationStateHandlePtr;
+
 	// If focus is lost on the popup, destroy it to prevent popups hanging around
-	PopupHandle = FSlateApplication::Get().OnFocusChanging().AddLambda([InFadeModeMenuWindow = FadeModeMenuWindow, InPopupHandle = PopupHandle](const FFocusEvent& FocusEvent, const FWeakWidgetPath& WeakWidgetPath
+	TSharedPtr<FDelegateHandle> PopupHandlePtr = MakeShared<FDelegateHandle>();
+	*PopupHandlePtr = FSlateApplication::Get().OnFocusChanging().AddLambda([InFadeModeMenuWindow = FadeModeMenuWindow, PopupHandlePtr](const FFocusEvent& FocusEvent, const FWeakWidgetPath& WeakWidgetPath
 		, const TSharedPtr<SWidget>& OldWidget, const FWidgetPath& WidgetPath, const TSharedPtr<SWidget>& NewWidget)
 		{
 			if (InFadeModeMenuWindow != nullptr && InFadeModeMenuWindow.IsValid())
@@ -663,11 +647,18 @@ void FWaveformTransformationTrimFadeRenderer::ShowSelectFadeModeMenuAtCursor(con
 					if (OldWidget && !OldWidget->IsHovered() && LockedFadeModeMenuWindow == OldWidget)
 					{
 						FSlateApplication::Get().RequestDestroyWindow(LockedFadeModeMenuWindow.ToSharedRef());
-						FSlateApplication::Get().OnFocusChanging().Remove(InPopupHandle);
+						FSlateApplication::Get().OnFocusChanging().Remove(*PopupHandlePtr);
 					}
 				}
 			}
 		});
+
+	if (PopupHandle.IsValid())
+	{
+		FSlateApplication::Get().OnFocusChanging().Remove(PopupHandle);
+	}
+
+	PopupHandle = *PopupHandlePtr;
 }
 
 const FSlateIcon FWaveformTransformationTrimFadeRenderer::GetFadeModeIcon(const EWaveEditorFadeMode& FadeMode)
@@ -781,12 +772,6 @@ FText FWaveformTransformationTrimFadeRenderer::GetPropertyEditedByCurrentInterac
 				FText::FromName(GET_MEMBER_NAME_CHECKED(FFadeFunctionData, FadeIn)),
 				FText::FromName(GET_MEMBER_NAME_CHECKED(UFadeFunction, Duration)));
 		}
-		else
-		{
-			PRAGMA_DISABLE_DEPRECATION_WARNINGS
-			OutPropertyName = FText::FromName(GET_MEMBER_NAME_CHECKED(UWaveformTransformationTrimFade, StartFadeTime));
-			PRAGMA_ENABLE_DEPRECATION_WARNINGS
-		}
 		break;
 	case FWaveformTransformationTrimFadeRenderer::ETrimFadeInteractionType::ScrubbingFadeOut:
 		if (StrongTrimFade->FadeFunctions.FadeOut)
@@ -795,12 +780,6 @@ FText FWaveformTransformationTrimFadeRenderer::GetPropertyEditedByCurrentInterac
 				FText::FromName(GET_MEMBER_NAME_CHECKED(FFadeFunctionData, FadeOut)),
 				FText::FromName(GET_MEMBER_NAME_CHECKED(UFadeFunction, Duration)));
 		}
-		else
-		{
-			PRAGMA_DISABLE_DEPRECATION_WARNINGS
-			OutPropertyName = FText::FromName(GET_MEMBER_NAME_CHECKED(UWaveformTransformationTrimFade, EndFadeTime));
-			PRAGMA_ENABLE_DEPRECATION_WARNINGS
-		}
 		break;
 	default:
 		break;
