Waveform Editor: Noise threshold for zero crossing logic

diff --git a/Engine/Plugins/Editor/WaveformEditor/Source/WaveformEditor/Private/WaveformEditor.cpp b/Engine/Plugins/Editor/WaveformEditor/Source/WaveformEditor/Private/WaveformEditor.cpp
--- a/Engine/Plugins/Editor/WaveformEditor/Source/WaveformEditor/Private/WaveformEditor.cpp
+++ b/Engine/Plugins/Editor/WaveformEditor/Source/WaveformEditor/Private/WaveformEditor.cpp
@@ -732,7 +732,12 @@ void FWaveformEditor::ModifyMarkerLoopRegion(ELoopModificationControls Modificat
 	TArray<float> PCMArray = WaveformView.DataProvider->GetTransformedPCMData();
 	check(PCMArray.Num() > 0);
 
-	MarkersTransformation->ModifyMarkerLoopRegion(Modification, PCMArray);
+	const UWaveformEditorTransformationsSettings* TransformationsSettings = GetWaveformEditorTransformationsSettings();
+	check(TransformationsSettings);
+
+	const float NoiseThresholdValue = powf(10.0f, TransformationsSettings->NoiseThreshold / 20.0f);
+
+	MarkersTransformation->ModifyMarkerLoopRegion(Modification, PCMArray, NoiseThresholdValue);
 }
 
 void FWaveformEditor::CycleMarkerLoopRegion(ELoopModificationControls Modification)
diff --git a/Engine/Plugins/Editor/WaveformEditor/Source/WaveformEditor/Private/WaveformEditorTransformationsSettings.h b/Engine/Plugins/Editor/WaveformEditor/Source/WaveformEditor/Private/WaveformEditorTransformationsSettings.h
--- a/Engine/Plugins/Editor/WaveformEditor/Source/WaveformEditor/Private/WaveformEditorTransformationsSettings.h
+++ b/Engine/Plugins/Editor/WaveformEditor/Source/WaveformEditor/Private/WaveformEditorTransformationsSettings.h
@@ -34,4 +34,7 @@ public:
 	 */
 	UPROPERTY(config, EditAnywhere, Category = "Launch Options")
 	TSet<TSubclassOf<UWaveformTransformationBase>> LaunchTransformations;
+
+	UPROPERTY(config, EditAnywhere, Category = "Zero Crossing Options", meta = (ClampMin = "-144.0", ClampMax = "0.0"))
+	float NoiseThreshold = -60.0f;
 };
diff --git a/Engine/Plugins/Editor/WaveformEditor/Source/WaveformTransformations/Private/WaveformTransformationMarkers.cpp b/Engine/Plugins/Editor/WaveformEditor/Source/WaveformTransformations/Private/WaveformTransformationMarkers.cpp
--- a/Engine/Plugins/Editor/WaveformEditor/Source/WaveformTransformations/Private/WaveformTransformationMarkers.cpp
+++ b/Engine/Plugins/Editor/WaveformEditor/Source/WaveformTransformations/Private/WaveformTransformationMarkers.cpp
@@ -189,7 +189,7 @@ void UWaveformTransformationMarkers::CycleMarkerLoopRegion(ELoopModificationCont
 }
 PRAGMA_ENABLE_DEPRECATION_WARNINGS
 
-void UWaveformTransformationMarkers::ModifyMarkerLoopRegion(ELoopModificationControls Modification, const TArray<float>& PCMArray)
+void UWaveformTransformationMarkers::ModifyMarkerLoopRegion(ELoopModificationControls Modification, const TArray<float>& PCMArray, float NoiseThreshold)
 {
 	check(Markers != nullptr);
 
@@ -260,7 +260,7 @@ void UWaveformTransformationMarkers::ModifyMarkerLoopRegion(ELoopModificationCon
 		const int64 StartFrame = (LoopRegion->FramePosition) * NumChannels;
 		const int64 MaxFrames = 0;
 
-		const int64 FramesToMove = GetFramesToNextZeroCrossing(StartFrame, MaxFrames, PCMArray, false);
+		const int64 FramesToMove = GetFramesToNextZeroCrossing(StartFrame, MaxFrames, PCMArray, false, NoiseThreshold);
 
 		LoopRegion->FramePosition -= FramesToMove;
 		LoopRegion->FrameLength += FramesToMove;
@@ -271,7 +271,7 @@ void UWaveformTransformationMarkers::ModifyMarkerLoopRegion(ELoopModificationCon
 		const int64 StartFrame = (LoopRegion->FramePosition) * NumChannels;
 		const int64 MaxFrames = (LoopRegion->FramePosition + LoopRegion->FrameLength) * NumChannels;
 
-		const int64 FramesToMove = GetFramesToNextZeroCrossing(StartFrame, MaxFrames, PCMArray, true);
+		const int64 FramesToMove = GetFramesToNextZeroCrossing(StartFrame, MaxFrames, PCMArray, true, NoiseThreshold);
 
 		LoopRegion->FramePosition += FramesToMove;
 		LoopRegion->FrameLength -= FramesToMove;
@@ -282,7 +282,7 @@ void UWaveformTransformationMarkers::ModifyMarkerLoopRegion(ELoopModificationCon
 		const int64 StartFrame = (LoopRegion->FramePosition + LoopRegion->FrameLength) * NumChannels;
 		const int64 MaxFrames = NumAvailableSamples;
 
-		LoopRegion->FrameLength += GetFramesToNextZeroCrossing(StartFrame, MaxFrames, PCMArray, true);
+		LoopRegion->FrameLength += GetFramesToNextZeroCrossing(StartFrame, MaxFrames, PCMArray, true, NoiseThreshold);
 	}
 	break;
 	case ELoopModificationControls::RightHandlePreviousZeroCrossing:
@@ -290,7 +290,7 @@ void UWaveformTransformationMarkers::ModifyMarkerLoopRegion(ELoopModificationCon
 		const int64 StartFrame = (LoopRegion->FramePosition + LoopRegion->FrameLength) * NumChannels;
 		const int64 MaxFrames = LoopRegion->FramePosition;
 
-		LoopRegion->FrameLength -= GetFramesToNextZeroCrossing(StartFrame, MaxFrames, PCMArray, false);
+		LoopRegion->FrameLength -= GetFramesToNextZeroCrossing(StartFrame, MaxFrames, PCMArray, false, NoiseThreshold);
 	}
 	default:
 		break;
@@ -378,19 +378,20 @@ void UWaveformTransformationMarkers::CycleMarkerLoopRegion(ELoopModificationCont
 	OnTransformationChanged.ExecuteIfBound(false);
 }
 
-int64 UWaveformTransformationMarkers::GetFramesToNextZeroCrossing(int64 StartFrame, int64 EndFrame, const TArray<float>& PCMArray, bool bCheckIncrementing) const
+int64 UWaveformTransformationMarkers::GetFramesToNextZeroCrossing(int64 StartFrame, int64 EndFrame, const TArray<float>& PCMArray, bool bCheckIncrementing, float NoiseThreshold) const
 {
 	check(NumChannels > 0);
 	int64 FramesMoved = 0;
 
 	bool bAllFramesPositiveLastFrame = false;
 	bool bAllFramesNegativeLastFrame = false;
+	bool bAllFramesLowNoiseLastFrame = false;
 
 	if (bCheckIncrementing)
 	{
 		for (int32 i = StartFrame; i < EndFrame; i += NumChannels)
 		{
-			if (CheckZeroCrossingFrame(i, PCMArray, bAllFramesPositiveLastFrame, bAllFramesNegativeLastFrame))
+			if (CheckZeroCrossingFrame(i, PCMArray, bAllFramesPositiveLastFrame, bAllFramesNegativeLastFrame, bAllFramesLowNoiseLastFrame, NoiseThreshold))
 			{
 				return FramesMoved;
 			}
@@ -402,7 +403,7 @@ int64 UWaveformTransformationMarkers::GetFramesToNextZeroCrossing(int64 StartFra
 	{
 		for (int32 i = StartFrame; i >= EndFrame; i -= NumChannels)
 		{
-			if (CheckZeroCrossingFrame(i, PCMArray, bAllFramesPositiveLastFrame, bAllFramesNegativeLastFrame))
+			if (CheckZeroCrossingFrame(i, PCMArray, bAllFramesPositiveLastFrame, bAllFramesNegativeLastFrame, bAllFramesLowNoiseLastFrame, NoiseThreshold))
 			{
 				return FramesMoved;
 			}
@@ -414,24 +415,26 @@ int64 UWaveformTransformationMarkers::GetFramesToNextZeroCrossing(int64 StartFra
 	return 0;
 }
 
-int64 UWaveformTransformationMarkers::CheckZeroCrossingFrame(int64 CurrentFrame, const TArray<float>& PCMArray, bool& bAllFramesPositiveLastFrame, bool& bAllFramesNegativeLastFrame) const
+int64 UWaveformTransformationMarkers::CheckZeroCrossingFrame(int64 CurrentFrame, const TArray<float>& PCMArray, bool& bAllFramesPositiveLastFrame, bool& bAllFramesNegativeLastFrame, bool& bAllFramesLowNoiseLastFrame, float NoiseThreshold) const
 {
 	check(NumChannels > 0);
 	check(CurrentFrame < PCMArray.Num());
 
 	int32 PositiveSampleCount = 0;
 	int32 NegativeSampleCount = 0;
+	int32 LowNoiseSampleCount = 0;
 
 	// check if all channels have a zero crossing
 	for (int32 ChannelNum = 0; ChannelNum < NumChannels; ChannelNum++)
 	{
 		const float FrameValue = PCMArray[CurrentFrame + ChannelNum];
 
-		// Treat 0 values as being valid for zero crossings
-		if (FrameValue == 0)
+		// Treat values within the noise threshold as valid for both negative and positive
+		if (FrameValue >= -NoiseThreshold && FrameValue <= NoiseThreshold)
 		{
 			++NegativeSampleCount;
 			++PositiveSampleCount;
+			++LowNoiseSampleCount;
 		}
 		else if (FrameValue < 0)
 		{
@@ -445,11 +448,14 @@ int64 UWaveformTransformationMarkers::CheckZeroCrossingFrame(int64 CurrentFrame,
 
 	bool bAllSamplesPositive = PositiveSampleCount == NumChannels;
 	bool bAllSamplesNegative = NegativeSampleCount == NumChannels;
+	bool bAllSamplesLowNoise = LowNoiseSampleCount == NumChannels;
 
 	// Check if we have crossed the zero crossing on all channels between last frame and this frame
-	bool bZeroCrossingOccurred = (bAllSamplesPositive && bAllFramesNegativeLastFrame) || (bAllSamplesNegative && bAllFramesPositiveLastFrame);
+	// If we are in a low noise frame then we ignore it for zero crossing
+	bool bZeroCrossingOccurred = (bAllSamplesPositive && bAllFramesNegativeLastFrame && !bAllSamplesLowNoise) || (bAllSamplesNegative && bAllFramesPositiveLastFrame && !bAllSamplesLowNoise);
 	bAllFramesPositiveLastFrame = bAllSamplesPositive;
 	bAllFramesNegativeLastFrame = bAllSamplesNegative;
+	bAllFramesLowNoiseLastFrame = bAllSamplesLowNoise;
 
 	return bZeroCrossingOccurred;
 }
diff --git a/Engine/Plugins/Editor/WaveformEditor/Source/WaveformTransformations/Public/WaveformTransformationMarkers.h b/Engine/Plugins/Editor/WaveformEditor/Source/WaveformTransformations/Public/WaveformTransformationMarkers.h
--- a/Engine/Plugins/Editor/WaveformEditor/Source/WaveformTransformations/Public/WaveformTransformationMarkers.h
+++ b/Engine/Plugins/Editor/WaveformEditor/Source/WaveformTransformations/Public/WaveformTransformationMarkers.h
@@ -110,7 +110,6 @@ public:
 	UPROPERTY(EditAnywhere, Category = "Loop Preview")
 	bool bIsPreviewingLoopRegion = false;
 
-	
 	UE_API virtual Audio::FTransformationPtr CreateTransformation() const override;
 
 	UE_API void UpdateConfiguration(FWaveTransformUObjectConfiguration& InOutConfiguration) override;
@@ -122,12 +121,12 @@ public:
 	UE_API void CycleMarkerLoopRegion(ELoopModificationControls Modification) const;
 
 	// Shift marker and loop region bounds with toolkit commands bound in the waveform editor
-	UE_API void ModifyMarkerLoopRegion(ELoopModificationControls Modification, const TArray<float> &PCMArray);
+	UE_API void ModifyMarkerLoopRegion(ELoopModificationControls Modification, const TArray<float> &PCMArray, float NoiseThreshold = 0.01f);
 	// Set SelectedCue by cycling through cues and loops with toolkit commands bound in the waveform editor
 	UE_API void CycleMarkerLoopRegion(ELoopModificationControls Modification);
 
-	UE_API int64 GetFramesToNextZeroCrossing(int64 StartFrame, int64 EndFrame, const TArray<float>& PCMArray, bool bCheckIncrementing) const;
-	UE_API int64 CheckZeroCrossingFrame(int64 CurrentFrame, const TArray<float>& PCMArray, bool& bAllFramesPositiveLastFrame, bool& bAllFramesNegativeLastFrame) const;
+	UE_API int64 GetFramesToNextZeroCrossing(int64 StartFrame, int64 EndFrame, const TArray<float>& PCMArray, bool bCheckIncrementing, float NoiseThreshold) const;
+	UE_API int64 CheckZeroCrossingFrame(int64 CurrentFrame, const TArray<float>& PCMArray, bool& bAllFramesPositiveLastFrame, bool& bAllFramesNegativeLastFrame, bool& bAllFramesLowNoiseLastFrame, float NoiseThreshold) const;
 
 	UE_API bool IsLoopRegionActive() const;
 	UE_API void SetLoopPreviewing();
