[Backout] - CL45071396
#rnx
Original CL Desc
-----------------------------------------------------------------
[Waveform Editor] - Cannot Undo/Redo Fade Function Duration and Curve parameters
[Waveform Editor] - Unable to use original Fade In/Out Curve and Duration parameters

diff --git a/Engine/Plugins/Editor/WaveformEditor/Source/WaveformEditor/Private/WaveformEditor.cpp b/Engine/Plugins/Editor/WaveformEditor/Source/WaveformEditor/Private/WaveformEditor.cpp
--- a/Engine/Plugins/Editor/WaveformEditor/Source/WaveformEditor/Private/WaveformEditor.cpp
+++ b/Engine/Plugins/Editor/WaveformEditor/Source/WaveformEditor/Private/WaveformEditor.cpp
@@ -261,13 +261,22 @@ FWaveformEditor::~FWaveformEditor()
 
 		if (FMath::IsNearlyEqual(TrimFadeTransformation->StartTime, 0, ErrorTolerance) && FMath::IsNearlyEqual(TrimFadeTransformation->EndTime, SoundWave->Duration, ErrorTolerance))
 		{
-			if ((TrimFadeTransformation->FadeFunctions.FadeIn && FMath::IsNearlyEqual(TrimFadeTransformation->FadeFunctions.FadeIn->Duration, 0, ErrorTolerance)
-				&& TrimFadeTransformation->FadeFunctions.FadeOut && FMath::IsNearlyEqual(TrimFadeTransformation->FadeFunctions.FadeOut->Duration, 0, ErrorTolerance))
-				|| (TrimFadeTransformation->FadeFunctions.FadeIn == nullptr && TrimFadeTransformation->FadeFunctions.FadeOut == nullptr))
+			bool bRemoveTrimFade = true;
+
+			if ((TrimFadeTransformation->FadeFunctions.FadeIn && !FMath::IsNearlyEqual(TrimFadeTransformation->FadeFunctions.FadeIn->Duration, 0, ErrorTolerance))
+				|| (TrimFadeTransformation->FadeFunctions.FadeOut && !FMath::IsNearlyEqual(TrimFadeTransformation->FadeFunctions.FadeOut->Duration, 0, ErrorTolerance)))
+			{
+				// If any FadeFunctions are present with a non-zero duration, do not remove the transformation
+				bRemoveTrimFade = false;
+			}
+			
+			PRAGMA_DISABLE_DEPRECATION_WARNINGS
+			// If no FadeFunctions are present, check the deprecated fade properties before removing the trimfade
+			if (bRemoveTrimFade && FMath::IsNearlyEqual(TrimFadeTransformation->StartFadeTime, 0, ErrorTolerance) && FMath::IsNearlyEqual(TrimFadeTransformation->EndFadeTime, 0, ErrorTolerance))
 			{
-				// If FadeIn and FadeOut Functions are present and they both have a zero duration or if both fade functions are null, remove the transformation
 				SoundWave->Transformations.Remove(TrimFadeTransformation);
 			}
+			PRAGMA_ENABLE_DEPRECATION_WARNINGS
 		}
 	}
 
@@ -368,9 +377,9 @@ TObjectPtr<UWaveformTransformationTrimFade> FWaveformEditor::GetOrAddTrimFadeTra
 		}
 
 		// Add default fade curves
-		TrimFadeTransformation->FadeFunctions.FadeIn = NewObject<UFadeFunction>(TrimFadeTransformation.Get(), FadeInType.Get(), NAME_None, RF_Transactional);
+		TrimFadeTransformation->FadeFunctions.FadeIn = NewObject<UFadeFunction>(TrimFadeTransformation.Get(), FadeInType.Get());
 		check(TrimFadeTransformation->FadeFunctions.FadeIn);
-		TrimFadeTransformation->FadeFunctions.FadeOut = NewObject<UFadeFunction>(TrimFadeTransformation.Get(), FadeOutType.Get(), NAME_None, RF_Transactional);
+		TrimFadeTransformation->FadeFunctions.FadeOut = NewObject<UFadeFunction>(TrimFadeTransformation.Get(), FadeOutType.Get());
 		check(TrimFadeTransformation->FadeFunctions.FadeOut);
 
 		TrimFadeTransformation->FadeFunctions.FadeIn->Duration = 0.f;
@@ -461,21 +470,24 @@ void FWaveformEditor::ToggleFadeIn()
 {
 	check(SoundWave);
 	check(GEditor != nullptr);
+	GEditor->BeginTransaction(TEXT("PropertyEditor"), LOCTEXT("SetNumericPropertyTransactionToggleFadeIn", "ToggleFadeIn"), nullptr);
+	SoundWave->Modify();
 
 	TObjectPtr<UWaveformTransformationTrimFade> TrimFadeTransformation = GetOrAddTrimFadeTransformation();
 	check(TrimFadeTransformation);
 
+	TrimFadeTransformation->Modify();
+	
 	if (TrimFadeTransformation->FadeFunctions.FadeIn)
 	{
-		GEditor->BeginTransaction(TEXT("PropertyEditor"), LOCTEXT("SetNumericPropertyTransactionToggleFadeIn", "ToggleFadeIn"), nullptr);
-		SoundWave->Modify();
-		TrimFadeTransformation->Modify();
-		TrimFadeTransformation->FadeFunctions.FadeIn->Modify();
-		
 		if (TrimFadeTransformation->FadeFunctions.FadeIn->Duration == 0.0f)
 		{
 			check(DefaultFadeInAmount >= 0.0f);
 
+			PRAGMA_DISABLE_DEPRECATION_WARNINGS
+			check(FadeInCurve >= UWaveformTransformationTrimFade::MinFadeCurve && FadeInCurve <= UWaveformTransformationTrimFade::MaxFadeCurve);
+			PRAGMA_ENABLE_DEPRECATION_WARNINGS
+
 			if (CachedFadeInAmount <= 0.0f)
 			{
 				CachedFadeInAmount = DefaultFadeInAmount;
@@ -490,10 +502,45 @@ void FWaveformEditor::ToggleFadeIn()
 			//No need to set FadeCurve when toggling off. This will preserve user curve data if they later manually change the StartFadeTime.
 			TrimFadeTransformation->FadeFunctions.FadeIn->Duration = 0.0f;
 		}
-		
-		NotifyPostTransformationChange(EPropertyChangeType::ValueSet);
-		GEditor->EndTransaction();
 	}
+	else // TrimFadeTransformation->FadeFunctions.FadeIn is nullptr
+	{
+		PRAGMA_DISABLE_DEPRECATION_WARNINGS
+		//if fade duration is 0, assume there is no fade present and toggle on the fade with the curve data selected in the menu
+		if (TrimFadeTransformation->StartFadeTime == 0.0f)
+		{
+			// No cached amount means we can instantiate a new FadeIn Function object
+			if (CachedFadeInAmount == 0.f)
+			{
+				TrimFadeTransformation->FadeFunctions.FadeIn = NewObject<UFadeFunction>(TrimFadeTransformation.Get(), FadeInType.Get());
+			}
+			else
+			{
+				check(DefaultFadeInAmount >= 0.0f);
+				check(FadeInCurve >= UWaveformTransformationTrimFade::MinFadeCurve && FadeInCurve <= UWaveformTransformationTrimFade::MaxFadeCurve);
+
+				if (CachedFadeInAmount <= 0.0f)
+				{
+					CachedFadeInAmount = DefaultFadeInAmount;
+				}
+
+				TrimFadeTransformation->StartFadeTime = CachedFadeInAmount;
+				TrimFadeTransformation->StartFadeCurve = FadeInCurve;
+			}
+		}
+		else //if fade duration is not 0, assume there is a fade applied and toggle fade "off" by setting fade time to 0
+		{
+			CachedFadeInAmount = TrimFadeTransformation->StartFadeTime;
+
+			//No need to set FadeCurve when toggling off. This will preserve user curve data if they later manually change the StartFadeTime.
+			TrimFadeTransformation->StartFadeTime = 0.0f;
+		}
+		PRAGMA_ENABLE_DEPRECATION_WARNINGS
+	}
+
+	NotifyPostTransformationChange(EPropertyChangeType::ValueSet);
+
+	GEditor->EndTransaction();
 }
 
 bool FWaveformEditor::CanFadeIn()
@@ -505,21 +552,24 @@ void FWaveformEditor::ToggleFadeOut()
 {
 	check(SoundWave);
 	check(GEditor != nullptr);
+	GEditor->BeginTransaction(TEXT("PropertyEditor"), LOCTEXT("SetNumericPropertyTransactionToggleFadeOut", "ToggleFadeOut"), nullptr);
+	SoundWave->Modify();
 
 	TObjectPtr<UWaveformTransformationTrimFade> TrimFadeTransformation = GetOrAddTrimFadeTransformation();
 	check(TrimFadeTransformation);
 
+	TrimFadeTransformation->Modify();
+
 	if (TrimFadeTransformation->FadeFunctions.FadeOut)
 	{
-		GEditor->BeginTransaction(TEXT("PropertyEditor"), LOCTEXT("SetNumericPropertyTransactionToggleFadeOut", "ToggleFadeOut"), nullptr);
-		SoundWave->Modify();
-		TrimFadeTransformation->Modify();
-		TrimFadeTransformation->FadeFunctions.FadeOut->Modify();
-		
 		if (TrimFadeTransformation->FadeFunctions.FadeOut->Duration == 0.0f)
 		{
 			check(DefaultFadeOutAmount >= 0.0f);
 
+			PRAGMA_DISABLE_DEPRECATION_WARNINGS
+			check(FadeOutCurve >= UWaveformTransformationTrimFade::MinFadeCurve && FadeOutCurve <= UWaveformTransformationTrimFade::MaxFadeCurve);
+			PRAGMA_ENABLE_DEPRECATION_WARNINGS
+
 			if (CachedFadeOutAmount <= 0.0f)
 			{
 				CachedFadeOutAmount = DefaultFadeOutAmount;
@@ -534,10 +584,44 @@ void FWaveformEditor::ToggleFadeOut()
 			//No need to set FadeCurve when toggling off. This will preserve user curve data if they later manually change the StartFadeTime.
 			TrimFadeTransformation->FadeFunctions.FadeOut->Duration = 0.0f;
 		}
+	}
+	else // TrimFadeTransformation->FadeFunctions.FadeOut is nullptr
+	{
+		PRAGMA_DISABLE_DEPRECATION_WARNINGS
+		//if fade duration is 0, assume there is no fade present and toggle on the fade with the curve data selected in the menu
+		if (TrimFadeTransformation->EndFadeTime == 0.0f)
+		{
+			// No cached amount means we can instantiate a new FadeIn Function object
+			if (CachedFadeOutAmount == 0.f)
+			{
+				TrimFadeTransformation->FadeFunctions.FadeOut = NewObject<UFadeFunction>(TrimFadeTransformation.Get(), FadeOutType.Get());
+			}
+			else
+			{
+				check(DefaultFadeOutAmount >= 0.0f);
+				check(FadeOutCurve >= UWaveformTransformationTrimFade::MinFadeCurve && FadeOutCurve <= UWaveformTransformationTrimFade::MaxFadeCurve);
 
-		NotifyPostTransformationChange(EPropertyChangeType::ValueSet);
-		GEditor->EndTransaction();
+				if (CachedFadeOutAmount <= 0.0f)
+				{
+					CachedFadeOutAmount = DefaultFadeOutAmount;
+				}
+
+				TrimFadeTransformation->EndFadeTime = CachedFadeOutAmount;
+				TrimFadeTransformation->EndFadeCurve = FadeOutCurve;
+			}
+		}
+		else //if fade duration is not 0, assume there is a fade applied and toggle fade "off" by setting fade time to 0
+		{
+			CachedFadeOutAmount = TrimFadeTransformation->EndFadeTime;
+
+			//No need to set FadeCurve when toggling off. This will preserve user curve data if they later manually change the EndFadeTime.
+			TrimFadeTransformation->EndFadeTime = 0.0f;
+		}
+		PRAGMA_ENABLE_DEPRECATION_WARNINGS
 	}
+
+	NotifyPostTransformationChange(EPropertyChangeType::ValueSet);
+	GEditor->EndTransaction();
 }
 
 bool FWaveformEditor::CanFadeOut()
@@ -945,10 +1029,7 @@ void FWaveformEditor::ExecuteOverwriteTransformations()
 {
 	for (TObjectPtr<UWaveformTransformationBase> Transformation : SoundWave->Transformations)
 	{
-		if (Transformation)
-		{
-			Transformation->OverwriteTransformation();
-		}
+		Transformation->OverwriteTransformation();
 	}
 }
 
@@ -1434,6 +1515,7 @@ bool FWaveformEditor::BindCommands()
 		Commands.FadeInLinear,
 		FExecuteAction::CreateLambda([this] 
 			{ 
+				FadeInCurve = UWaveformTransformationTrimFade::FadeModeToCurveValueMap[EWaveEditorFadeMode::Linear];
 				FadeInType = UWaveformTransformationTrimFade::FadeModeToFadeFunctionMap[EWaveEditorFadeMode::Linear];
 			}),
 		FCanExecuteAction::CreateLambda([this] { return true; }),
@@ -1443,6 +1525,7 @@ bool FWaveformEditor::BindCommands()
 		Commands.FadeInExponential,
 		FExecuteAction::CreateLambda([this] 
 			{ 
+				FadeInCurve = UWaveformTransformationTrimFade::FadeModeToCurveValueMap[EWaveEditorFadeMode::Exponential];
 				FadeInType = UWaveformTransformationTrimFade::FadeModeToFadeFunctionMap[EWaveEditorFadeMode::Exponential];
 			}),
 		FCanExecuteAction::CreateLambda([this] { return true; }),
@@ -1452,6 +1535,7 @@ bool FWaveformEditor::BindCommands()
 		Commands.FadeInLogarithmic,
 		FExecuteAction::CreateLambda([this] 
 			{ 
+				FadeInCurve = UWaveformTransformationTrimFade::FadeModeToCurveValueMap[EWaveEditorFadeMode::Logarithmic];
 				FadeInType = UWaveformTransformationTrimFade::FadeModeToFadeFunctionMap[EWaveEditorFadeMode::Logarithmic];
 			}),
 		FCanExecuteAction::CreateLambda([this] { return true; }),
@@ -1461,6 +1545,7 @@ bool FWaveformEditor::BindCommands()
 		Commands.FadeInSigmoid,
 		FExecuteAction::CreateLambda([this] 
 			{ 
+				FadeInCurve = UWaveformTransformationTrimFade::FadeModeToCurveValueMap[EWaveEditorFadeMode::Sigmoid];
 				FadeInType = UWaveformTransformationTrimFade::FadeModeToFadeFunctionMap[EWaveEditorFadeMode::Sigmoid];
 			}),
 		FCanExecuteAction::CreateLambda([this] { return true; }),
@@ -1475,6 +1560,7 @@ bool FWaveformEditor::BindCommands()
 		Commands.FadeOutLinear,
 		FExecuteAction::CreateLambda([this] 
 			{ 
+				FadeOutCurve = UWaveformTransformationTrimFade::FadeModeToCurveValueMap[EWaveEditorFadeMode::Linear];
 				FadeOutType = UWaveformTransformationTrimFade::FadeModeToFadeFunctionMap[EWaveEditorFadeMode::Linear];
 			}),
 		FCanExecuteAction::CreateLambda([this] {return true; }),
@@ -1484,6 +1570,7 @@ bool FWaveformEditor::BindCommands()
 		Commands.FadeOutExponential,
 		FExecuteAction::CreateLambda([this] 
 			{ 
+				FadeOutCurve = UWaveformTransformationTrimFade::FadeModeToCurveValueMap[EWaveEditorFadeMode::Exponential];
 				FadeOutType = UWaveformTransformationTrimFade::FadeModeToFadeFunctionMap[EWaveEditorFadeMode::Exponential];
 			}),
 		FCanExecuteAction::CreateLambda([this] {return true; }),
@@ -1493,6 +1580,7 @@ bool FWaveformEditor::BindCommands()
 		Commands.FadeOutLogarithmic,
 		FExecuteAction::CreateLambda([this] 
 			{ 
+				FadeOutCurve = UWaveformTransformationTrimFade::FadeModeToCurveValueMap[EWaveEditorFadeMode::Logarithmic];
 				FadeOutType = UWaveformTransformationTrimFade::FadeModeToFadeFunctionMap[EWaveEditorFadeMode::Logarithmic];
 			}),
 		FCanExecuteAction::CreateLambda([this] {return true; }),
@@ -1502,6 +1590,7 @@ bool FWaveformEditor::BindCommands()
 		Commands.FadeOutSigmoid,
 		FExecuteAction::CreateLambda([this] 
 			{ 
+				FadeOutCurve = UWaveformTransformationTrimFade::FadeModeToCurveValueMap[EWaveEditorFadeMode::Sigmoid];
 				FadeOutType = UWaveformTransformationTrimFade::FadeModeToFadeFunctionMap[EWaveEditorFadeMode::Sigmoid];
 			}),
 		FCanExecuteAction::CreateLambda([this] {return true; }),
@@ -1614,23 +1703,14 @@ void FWaveformEditor::OnAssetReimport(UObject* ReimportedObject, bool bSuccessfu
 
 	if (ReimportedObject->IsA<USoundWave>() && ReimportedObject->GetPathName() == SoundWave->GetPathName())
 	{
-		check(SoundWave);
-
 		//If the waveform editor is open, updating Transformations on reimport will show
 		// the effects of the overwritten transformations immediately instead of after a
 		// change is made by the user.
 		TransformationChainConfig = SoundWave->UpdateTransformations();
 		CreateWaveformView();
-
-		check(WaveformView.DataProvider);
-		check(WaveformView.ViewWidget);
 		WaveformView.DataProvider->RequestSequenceView(TransportCoordinator->GetDisplayRange());
 		WaveformView.ViewWidget->SetPlayheadRatio(TransportCoordinator->GetFocusPoint());
-
-		if (TSharedPtr<SDockTab> LiveTab = this->TabManager->FindExistingLiveTab(WaveformDisplayTabId))
-		{
-			LiveTab->SetContent(WaveformView.ViewWidget.ToSharedRef());
-		}
+		this->TabManager->FindExistingLiveTab(WaveformDisplayTabId)->SetContent(WaveformView.ViewWidget.ToSharedRef());
 	}
 }
 
@@ -1842,19 +1922,19 @@ void FWaveformEditor::PostRedo(bool bSuccess)
 
 bool FWaveformEditor::MatchesContext(const FTransactionContext& InContext, const TArray<TPair<UObject*, FTransactionObjectEvent>>& TransactionObjectContexts) const
 {
-	bool bShouldUndo = false;
+	bool bShoouldUndo = false;
 	for (const TPair<UObject*, FTransactionObjectEvent>& TransactionObjectPair : TransactionObjectContexts)
 	{
-		if (const UObject* Object = TransactionObjectPair.Key)
+		UObject* Object = TransactionObjectPair.Key;
+		const UClass* ObjectClass = Object->GetClass();
+
+		if (ObjectClass)
 		{
-			if (const UClass* ObjectClass = Object->GetClass())
-			{
-				bShouldUndo |= ObjectClass->IsChildOf(UWaveformTransformationBase::StaticClass()) || ObjectClass->IsChildOf(USoundWave::StaticClass());
-			}
+			bShoouldUndo = ObjectClass->IsChildOf(UWaveformTransformationBase::StaticClass()) || ObjectClass->IsChildOf(USoundWave::StaticClass());
 		}
 	}
 
-	return bShouldUndo;
+	return bShoouldUndo;
 }
 
 void FWaveformEditor::InitToolMenuContext(FToolMenuContext& MenuContext)
diff --git a/Engine/Plugins/Editor/WaveformEditor/Source/WaveformEditor/Public/WaveformEditor.h b/Engine/Plugins/Editor/WaveformEditor/Source/WaveformEditor/Public/WaveformEditor.h
--- a/Engine/Plugins/Editor/WaveformEditor/Source/WaveformEditor/Public/WaveformEditor.h
+++ b/Engine/Plugins/Editor/WaveformEditor/Source/WaveformEditor/Public/WaveformEditor.h
@@ -179,10 +179,12 @@ private:
 
 	static constexpr float DefaultFadeInAmount = 0.5f;
 	float CachedFadeInAmount = 0.f;
+	float FadeInCurve = UWaveformTransformationTrimFade::FadeModeToCurveValueMap[EWaveEditorFadeMode::Exponential];
 	TSubclassOf<UFadeFunction> FadeInType = UWaveformTransformationTrimFade::FadeModeToFadeFunctionMap[EWaveEditorFadeMode::Exponential];
 
 	static constexpr float DefaultFadeOutAmount = 0.5f;
 	float CachedFadeOutAmount = 0.f;
+	float FadeOutCurve = UWaveformTransformationTrimFade::FadeModeToCurveValueMap[EWaveEditorFadeMode::Exponential];
 	TSubclassOf<UFadeFunction> FadeOutType = UWaveformTransformationTrimFade::FadeModeToFadeFunctionMap[EWaveEditorFadeMode::Exponential];
 
 	// Tracking bool to ensure we detect wave cue changes and regenerate transformations
diff --git a/Engine/Plugins/Editor/WaveformEditor/Source/WaveformTransformations/Private/WaveformTransformationTrimFade.cpp b/Engine/Plugins/Editor/WaveformEditor/Source/WaveformTransformations/Private/WaveformTransformationTrimFade.cpp
--- a/Engine/Plugins/Editor/WaveformEditor/Source/WaveformTransformations/Private/WaveformTransformationTrimFade.cpp
+++ b/Engine/Plugins/Editor/WaveformEditor/Source/WaveformTransformations/Private/WaveformTransformationTrimFade.cpp
@@ -1,8 +1,6 @@
 // Copyright Epic Games, Inc. All Rights Reserved.
 
 #include "WaveformTransformationTrimFade.h"
-
-#include "Editor.h"
 #include "Sound/SoundWave.h"
 
 #include UE_INLINE_GENERATED_CPP_BY_NAME(WaveformTransformationTrimFade)
@@ -11,30 +9,16 @@ namespace
 {
 	// Makes the two functions for SCurve meet at (0.5, 0.5)
 	const double XMultiplier = FMath::Sqrt(10.0) / 5.0;
-
-	UE_DEPRECATED(5.7, "Individual fade properties have been deprecated. Use FadeFunctions::FadeIn::GetFadeInCurveValue instead.")
-	void SetFadeCurve(TObjectPtr<UFadeFunction> FadeFunction, float FadeCurve, float Sharpness)
-	{
-		if (FadeFunction->IsA<UFadeCurveFunction>())
-		{
-			TObjectPtr<UFadeCurveFunction> CurveFunction = CastChecked<UFadeCurveFunction>(FadeFunction);
-
-			if (CurveFunction)
-			{
-				if (!CurveFunction->IsA<UFadeCurveFunctionSigmoid>())
-				{
-					CurveFunction->SetFadeCurve(FadeCurve);
-				}
-				else
-				{
-					CurveFunction->SetFadeCurve(Sharpness);
-				}
-			}
-			
-		}
-	}
 }
 
+const TMap<EWaveEditorFadeMode, TSubclassOf<UFadeFunction>> UWaveformTransformationTrimFade::FadeModeToFadeFunctionMap =
+{
+	{EWaveEditorFadeMode::Linear, UFadeFunctionLinear::StaticClass()},
+	{EWaveEditorFadeMode::Exponential, UFadeCurveFunctionExponential::StaticClass()},
+	{EWaveEditorFadeMode::Logarithmic, UFadeCurveFunctionLogarithmic::StaticClass()},
+	{EWaveEditorFadeMode::Sigmoid, UFadeCurveFunctionSigmoid::StaticClass()}
+};
+
 PRAGMA_DISABLE_DEPRECATION_WARNINGS
 const TMap<EWaveEditorFadeMode, float> UWaveformTransformationTrimFade::FadeModeToCurveValueMap =
 {
@@ -45,35 +29,20 @@ const TMap<EWaveEditorFadeMode, float> UWaveformTransformationTrimFade::FadeMode
 };
 PRAGMA_ENABLE_DEPRECATION_WARNINGS
 
-const TMap<EWaveEditorFadeMode, TSubclassOf<UFadeFunction>> UWaveformTransformationTrimFade::FadeModeToFadeFunctionMap =
-{
-	{EWaveEditorFadeMode::Linear, UFadeFunctionLinear::StaticClass()},
-	{EWaveEditorFadeMode::Exponential, UFadeCurveFunctionExponential::StaticClass()},
-	{EWaveEditorFadeMode::Logarithmic, UFadeCurveFunctionLogarithmic::StaticClass()},
-	{EWaveEditorFadeMode::Sigmoid, UFadeCurveFunctionSigmoid::StaticClass()}
-};
-
 const float UWaveformTransformationTrimFade::MinFadeCurve = -0.1f;
 const float UWaveformTransformationTrimFade::MaxFadeCurve = 10.f;
 
-
 namespace WaveformTransformationTrimFadeNames
 {
 	static FLazyName StartTimeName(GET_MEMBER_NAME_CHECKED(UWaveformTransformationTrimFade, StartTime));
 	static FLazyName EndTimeName(GET_MEMBER_NAME_CHECKED(UWaveformTransformationTrimFade, EndTime));
 }
 
-static void ApplyFadeIn(Audio::FAlignedFloatBuffer& InputAudio, const FFadeFunctionData& FadeFunctions, const int32 NumChannels, const float SampleRate)
+static void ApplyFadeIn(Audio::FAlignedFloatBuffer& InputAudio, const FFadeFunctionData& FadeFunctions, const float FadeLength, const float FadeCurve, const float SCurveSharpness, const int32 NumChannels, const float SampleRate)
 {
-	if (FadeFunctions.FadeIn == nullptr)
-	{
-		return;
-	}
-
 	check(NumChannels > 0);
 	check(SampleRate > 0.f);
-
-	const float FadeDuration = FadeFunctions.FadeIn->Duration;
+	const float FadeDuration = FadeFunctions.FadeIn ? FadeFunctions.FadeIn->Duration : FadeLength;
 	const int32 FadeNumFrames = FMath::Min(FadeDuration * SampleRate, InputAudio.Num() / NumChannels);
 
 	if (InputAudio.Num() < NumChannels || FadeDuration < SMALL_NUMBER || FadeNumFrames == 0)
@@ -89,7 +58,10 @@ static void ApplyFadeIn(Audio::FAlignedFloatBuffer& InputAudio, const FFadeFunct
 		check(FadeFraction <= 1.0f);
 		check(FadeFraction >= 0.0f);
 
-		const float EnvValue = FadeFunctions.FadeIn->GetFadeInCurveValue(FadeFraction);
+		PRAGMA_DISABLE_DEPRECATION_WARNINGS
+		const float EnvValue = FadeFunctions.FadeIn ? FadeFunctions.FadeIn->GetFadeInCurveValue(FadeFraction) :
+		UWaveformTransformationTrimFade::GetFadeInCurveValue(FadeCurve, FadeFraction, SCurveSharpness);
+		PRAGMA_ENABLE_DEPRECATION_WARNINGS
 
 		for (int32 ChannelIt = 0; ChannelIt < NumChannels; ++ChannelIt)
 		{
@@ -100,17 +72,11 @@ static void ApplyFadeIn(Audio::FAlignedFloatBuffer& InputAudio, const FFadeFunct
 	}
 }
 
-static void ApplyFadeOut(Audio::FAlignedFloatBuffer& InputAudio, const FFadeFunctionData& FadeFunctions, const int32 NumChannels, const float SampleRate)
+static void ApplyFadeOut(Audio::FAlignedFloatBuffer& InputAudio, const FFadeFunctionData& FadeFunctions, const float FadeLength, const float FadeCurve, const float SCurveSharpness, const int32 NumChannels, const float SampleRate)
 {
-	if (FadeFunctions.FadeOut == nullptr)
-	{
-		return;
-	}
-
 	check(NumChannels > 0);
 	check(SampleRate > 0.f);
-
-	const float FadeDuration = FadeFunctions.FadeOut->Duration;
+	const float FadeDuration = FadeFunctions.FadeOut ? FadeFunctions.FadeOut->Duration : FadeLength;
 	const int32 FadeNumFrames = FMath::Min(FadeDuration * SampleRate, InputAudio.Num() / NumChannels);
 
 	if (InputAudio.Num() < NumChannels || FadeDuration < SMALL_NUMBER || FadeNumFrames == 0)
@@ -127,7 +93,10 @@ static void ApplyFadeOut(Audio::FAlignedFloatBuffer& InputAudio, const FFadeFunc
 		check(FadeFraction <= 1.0f);
 		check(FadeFraction >= 0.0f);
 
-		const float EnvValue = FadeFunctions.FadeOut->GetFadeOutCurveValue(FadeFraction);
+		PRAGMA_DISABLE_DEPRECATION_WARNINGS
+		const float EnvValue = FadeFunctions.FadeOut ? FadeFunctions.FadeOut->GetFadeOutCurveValue(FadeFraction) :
+		UWaveformTransformationTrimFade::GetFadeOutCurveValue(FadeCurve, FadeFraction, SCurveSharpness);
+		PRAGMA_ENABLE_DEPRECATION_WARNINGS
 
 		for (int32 ChannelIt = 0; ChannelIt < NumChannels; ++ChannelIt)
 		{
@@ -137,11 +106,18 @@ static void ApplyFadeOut(Audio::FAlignedFloatBuffer& InputAudio, const FFadeFunc
 		InputPtr += NumChannels;
 	}
 }
-
-FWaveTransformationTrimFade::FWaveTransformationTrimFade(const FFadeFunctionData& InFadeFunctions, double InStartTime, double InEndTime)
+PRAGMA_DISABLE_DEPRECATION_WARNINGS
+FWaveTransformationTrimFade::FWaveTransformationTrimFade(const FFadeFunctionData& InFadeFunctions, double InStartTime, double InEndTime, float InStartFadeTime, float InStartFadeCurve, float InStartSCurveSharpness, float InEndFadeTime, float InEndFadeCurve, float InEndSCurveSharpness)
 	: FadeFunctions(InFadeFunctions)
 	, StartTime(InStartTime)
-	, EndTime(InEndTime){}
+	, EndTime(InEndTime)
+	, StartFadeTime(InStartFadeTime)
+	, StartFadeCurve(FMath::Max(InStartFadeCurve, UWaveformTransformationTrimFade::MinFadeCurve))
+	, StartSCurveSharpness(InStartSCurveSharpness)
+	, EndFadeTime(InEndFadeTime)
+	, EndFadeCurve(FMath::Max(InEndFadeCurve, UWaveformTransformationTrimFade::MinFadeCurve))
+	, EndSCurveSharpness(InEndSCurveSharpness){}
+PRAGMA_ENABLE_DEPRECATION_WARNINGS
 
 void FWaveTransformationTrimFade::ProcessAudio(Audio::FWaveformTransformationWaveInfo& InOutWaveInfo) const
 {
@@ -149,39 +125,28 @@ void FWaveTransformationTrimFade::ProcessAudio(Audio::FWaveformTransformationWav
 
 	Audio::FAlignedFloatBuffer& InputAudio = *InOutWaveInfo.Audio;
 
-	// If InputAudio is smaller than the size of a frame, do not process the audio.
-	if (InputAudio.Num() < InOutWaveInfo.NumChannels)
+	if(InputAudio.Num() == 0)
 	{
 		return;
 	}
 
-	check(InOutWaveInfo.NumChannels > 0);
-
-	const int64 LastInputAudioIndex = static_cast<int64>(InputAudio.Num() - 1);
-	const int64 NumChannelsMinusOne = static_cast<int64>(InOutWaveInfo.NumChannels - 1); // Used to step backwards from the end sample to a valid frame
-	const int64 FirstSampleOfLastFrame = LastInputAudioIndex - NumChannelsMinusOne; // The last sample that begins a frame
-	int64 StartSample = FMath::Min(FMath::FloorToInt64(StartTime * InOutWaveInfo.SampleRate) * InOutWaveInfo.NumChannels, FirstSampleOfLastFrame);
-	int64 EndSample = LastInputAudioIndex;
+	const int64 LastInputAudioIndex = InputAudio.Num() - 1;
+	int64 StartSample = FMath::Min(FMath::FloorToInt64(StartTime * InOutWaveInfo.SampleRate) * InOutWaveInfo.NumChannels, LastInputAudioIndex - InOutWaveInfo.NumChannels);
+	int64 EndSample = InputAudio.Num();
 	
 	if(EndTime > 0.f)
 	{	
 		const int64 EndFrame = FMath::RoundToInt64(EndTime * InOutWaveInfo.SampleRate);
-		EndSample = EndFrame * InOutWaveInfo.NumChannels + NumChannelsMinusOne;
+		EndSample = EndFrame * InOutWaveInfo.NumChannels;
 		EndSample = FMath::Min(EndSample, LastInputAudioIndex);
 	}
 
-	if (StartSample > EndSample - NumChannelsMinusOne)
+	if (StartSample > EndSample - InOutWaveInfo.NumChannels)
 	{
-		StartSample = FMath::Max(EndSample - NumChannelsMinusOne, 0);
+		StartSample = EndSample - InOutWaveInfo.NumChannels;
 	}
 
-	check(StartSample >= 0);
-	check(StartSample % InOutWaveInfo.NumChannels == 0);
-	check(EndSample >= StartSample);
-	check((EndSample - NumChannelsMinusOne) % InOutWaveInfo.NumChannels == 0);
-
-	const int64 FinalSize = (EndSample - StartSample) + 1;
-	check(FinalSize % InOutWaveInfo.NumChannels == 0);
+	const int64 FinalSize = EndSample - StartSample;
 
 	InOutWaveInfo.StartFrameOffset = StartSample - (StartSample % InOutWaveInfo.NumChannels);
 	InOutWaveInfo.NumEditedSamples = FinalSize;
@@ -191,9 +156,7 @@ void FWaveTransformationTrimFade::ProcessAudio(Audio::FWaveformTransformationWav
 		return;
 	}
 
-	const bool bProcessFadeIn = FadeFunctions.FadeIn && FadeFunctions.FadeIn->Duration > 0.f;
-	const bool bProcessFadeOut = FadeFunctions.FadeOut && FadeFunctions.FadeOut->Duration > 0.f;
-	const bool bProcessFades = bProcessFadeIn || bProcessFadeOut;
+	const bool bProcessFades = StartFadeTime > 0.f || EndFadeTime > 0.f;
 
 	if (!bProcessFades && FinalSize == InputAudio.Num())
 	{
@@ -210,15 +173,17 @@ void FWaveTransformationTrimFade::ProcessAudio(Audio::FWaveformTransformationWav
 
 	FMemory::Memcpy(InputAudio.GetData(), TempBuffer.GetData(), FinalSize * sizeof(float));
 
-	if (bProcessFadeIn)
+	PRAGMA_DISABLE_DEPRECATION_WARNINGS
+	if ((FadeFunctions.FadeIn && FadeFunctions.FadeIn->Duration > 0.f) || StartFadeTime > 0.f)
 	{
-		ApplyFadeIn(InputAudio, FadeFunctions, InOutWaveInfo.NumChannels, InOutWaveInfo.SampleRate);
+		ApplyFadeIn(InputAudio, FadeFunctions, StartFadeTime, StartFadeCurve, StartSCurveSharpness, InOutWaveInfo.NumChannels, InOutWaveInfo.SampleRate);
 	}
 
-	if (bProcessFadeOut)
+	if ((FadeFunctions.FadeOut && FadeFunctions.FadeOut->Duration > 0.f) || EndFadeTime > 0.f)
 	{
-		ApplyFadeOut(InputAudio, FadeFunctions, InOutWaveInfo.NumChannels, InOutWaveInfo.SampleRate);
-	}
+		ApplyFadeOut(InputAudio, FadeFunctions, EndFadeTime, EndFadeCurve, EndSCurveSharpness, InOutWaveInfo.NumChannels, InOutWaveInfo.SampleRate);
+	}	
+	PRAGMA_ENABLE_DEPRECATION_WARNINGS
 }
 
 void UWaveformTransformationTrimFade::PostEditChangeProperty(FPropertyChangedEvent& PropertyChangedEvent)
@@ -227,6 +192,11 @@ void UWaveformTransformationTrimFade::PostEditChangeProperty(FPropertyChangedEve
 
 	const float MaxDuration = EndTime - StartTime;
 
+	PRAGMA_DISABLE_DEPRECATION_WARNINGS
+	EndFadeTime = FMath::Min(EndFadeTime, MaxDuration);
+	StartFadeTime = FMath::Min(StartFadeTime, MaxDuration);
+	PRAGMA_ENABLE_DEPRECATION_WARNINGS
+
 	if (FadeFunctions.FadeIn)
 	{
 		FadeFunctions.FadeIn->Duration = FMath::Min(FadeFunctions.FadeIn->Duration, MaxDuration);
@@ -240,7 +210,9 @@ void UWaveformTransformationTrimFade::PostEditChangeProperty(FPropertyChangedEve
 
 Audio::FTransformationPtr UWaveformTransformationTrimFade::CreateTransformation() const
 {
-	return MakeUnique<FWaveTransformationTrimFade>(FadeFunctions, StartTime, EndTime);
+	PRAGMA_DISABLE_DEPRECATION_WARNINGS
+	return MakeUnique<FWaveTransformationTrimFade>(FadeFunctions, StartTime, EndTime, StartFadeTime, StartFadeCurve, StartSCurveSharpness, EndFadeTime, EndFadeCurve, EndSCurveSharpness);
+	PRAGMA_ENABLE_DEPRECATION_WARNINGS
 }
 
 void UWaveformTransformationTrimFade::UpdateConfiguration(FWaveTransformUObjectConfiguration& InOutConfiguration)
@@ -251,7 +223,6 @@ void UWaveformTransformationTrimFade::UpdateConfiguration(FWaveTransformUObjectC
 	InOutConfiguration.EndTime = EndTime;
 }
 
-PRAGMA_DISABLE_DEPRECATION_WARNINGS
 const double UWaveformTransformationTrimFade::GetFadeInCurveValue(const float StartFadeCurve, const double FadeFraction, const float SCurveSharpness)
 {
 	double CurveValue = 1;
@@ -282,7 +253,7 @@ const double UWaveformTransformationTrimFade::GetFadeInCurveValue(const float St
 const double UWaveformTransformationTrimFade::GetFadeOutCurveValue(const float EndFadeCurve, const double FadeFraction, const float SCurveSharpness)
 {
 	double CurveValue = 1;
-
+	
 	if (EndFadeCurve < 0)
 	{
 		const double X = FadeFraction;
@@ -307,103 +278,6 @@ const double UWaveformTransformationTrimFade::GetFadeOutCurveValue(const float E
 
 	return CurveValue;
 }
-PRAGMA_ENABLE_DEPRECATION_WARNINGS
-
-void UWaveformTransformationTrimFade::PostInitProperties()
-{
-	Super::PostInitProperties();
-
-	if (FadeFunctions.FadeIn == nullptr)
-	{
-		PRAGMA_DISABLE_DEPRECATION_WARNINGS
-		if (StartFadeTime == 0.f)
-		PRAGMA_ENABLE_DEPRECATION_WARNINGS
-		{
-			FadeFunctions.FadeIn = NewObject<UFadeCurveFunctionExponential>(this, UFadeCurveFunctionExponential::StaticClass(), NAME_None, RF_Transactional);
-			FadeFunctions.FadeIn->Duration = 0.f;
-		}
-	}
-
-	check(FadeFunctions.FadeIn);
-
-	if (FadeFunctions.FadeOut == nullptr)
-	{
-		PRAGMA_DISABLE_DEPRECATION_WARNINGS
-		if (EndFadeTime == 0.f)
-		PRAGMA_ENABLE_DEPRECATION_WARNINGS
-		{
-			FadeFunctions.FadeOut = NewObject<UFadeCurveFunctionExponential>(this, UFadeCurveFunctionExponential::StaticClass(), NAME_None, RF_Transactional);
-			FadeFunctions.FadeOut->Duration = 0.f;
-		}
-	}
-
-	check(FadeFunctions.FadeOut);
-}
-
-void UWaveformTransformationTrimFade::PostLoad()
-{
-	Super::PostLoad();
-
-	const float Linear = 1.f;
-	
-	PRAGMA_DISABLE_DEPRECATION_WARNINGS
-	if (StartFadeTime != 0.f)
-	{
-		if (FadeFunctions.FadeIn == nullptr)
-		{
-			if (StartFadeCurve > Linear)
-			{
-				FadeFunctions.FadeIn = NewObject<UFadeCurveFunctionExponential>(this, UFadeCurveFunctionExponential::StaticClass(), NAME_None, RF_Transactional);
-			}
-			else if (StartFadeCurve < 0.f)
-			{
-				FadeFunctions.FadeIn = NewObject<UFadeCurveFunctionSigmoid>(this, UFadeCurveFunctionSigmoid::StaticClass(), NAME_None, RF_Transactional);
-			}
-			else if (StartFadeCurve < Linear)
-			{
-				FadeFunctions.FadeIn = NewObject<UFadeCurveFunctionLogarithmic>(this, UFadeCurveFunctionLogarithmic::StaticClass(), NAME_None, RF_Transactional);
-			}
-			else if (StartFadeCurve == Linear)
-			{
-				FadeFunctions.FadeIn = NewObject<UFadeFunctionLinear>(this, UFadeFunctionLinear::StaticClass(), NAME_None, RF_Transactional);
-			}
-			
-			FadeFunctions.FadeIn->Duration = StartFadeTime;
-			SetFadeCurve(FadeFunctions.FadeIn, StartFadeCurve, StartSCurveSharpness);
-		}
-
-		StartFadeTime = 0.f;
-	}
-
-	if (EndFadeTime != 0.f)
-	{
-		if (FadeFunctions.FadeOut == nullptr)
-		{
-			if (EndFadeCurve > Linear)
-			{
-				FadeFunctions.FadeOut = NewObject<UFadeCurveFunctionExponential>(this, UFadeCurveFunctionExponential::StaticClass(), NAME_None, RF_Transactional);
-			}
-			else if (EndFadeCurve < 0.f)
-			{
-				FadeFunctions.FadeOut = NewObject<UFadeCurveFunctionSigmoid>(this, UFadeCurveFunctionSigmoid::StaticClass(), NAME_None, RF_Transactional);
-			}
-			else if (EndFadeCurve < Linear)
-			{
-				FadeFunctions.FadeOut = NewObject<UFadeCurveFunctionLogarithmic>(this, UFadeCurveFunctionLogarithmic::StaticClass(), NAME_None, RF_Transactional);
-			}
-			else if (EndFadeCurve == Linear)
-			{
-				FadeFunctions.FadeOut = NewObject<UFadeFunctionLinear>(this, UFadeFunctionLinear::StaticClass(), NAME_None, RF_Transactional);
-			}
-			
-			FadeFunctions.FadeOut->Duration = EndFadeTime;
-			SetFadeCurve(FadeFunctions.FadeOut, EndFadeCurve, EndSCurveSharpness);
-		}
-
-		EndFadeTime = 0.f;
-	}
-	PRAGMA_ENABLE_DEPRECATION_WARNINGS
-}
 
 void UWaveformTransformationTrimFade::UpdateDurationProperties(const float InAvailableDuration)
 {
@@ -520,18 +394,3 @@ const double UFadeCurveFunctionSigmoid::GetFadeOutCurveValue(const double FadeFr
 
 	return CurveValue;
 }
-
-void UFadeFunction::PostEditUndo()
-{
-	Super::PostEditUndo();
-
-	UObject* Outer = GetOuter();
-	if (Outer && Outer->IsA<UWaveformTransformationBase>())
-	{
-		UWaveformTransformationBase* Transformation = CastChecked<UWaveformTransformationBase>(GetOuter());
-
-		static const bool bMarkFileDirty = true;
-
-		Transformation->OnTransformationChanged.ExecuteIfBound(bMarkFileDirty);
-	}
-}
diff --git a/Engine/Plugins/Editor/WaveformEditor/Source/WaveformTransformations/Public/WaveformTransformationTrimFade.h b/Engine/Plugins/Editor/WaveformEditor/Source/WaveformTransformations/Public/WaveformTransformationTrimFade.h
--- a/Engine/Plugins/Editor/WaveformEditor/Source/WaveformTransformations/Public/WaveformTransformationTrimFade.h
+++ b/Engine/Plugins/Editor/WaveformEditor/Source/WaveformTransformations/Public/WaveformTransformationTrimFade.h
@@ -20,11 +20,6 @@ public:
 
 	UPROPERTY(EditAnywhere, Category = "Fade", BlueprintReadWrite, meta = (ClampMin = 0.0, DisplayName = "Fade Duration"))
 	float Duration = 0.5f;
-
-private: 
-#if WITH_EDITOR
-	virtual void PostEditUndo() override;
-#endif // WITH_EDITOR
 };
 
 UCLASS(MinimalAPI, Abstract, EditInlineNew, DefaultToInstanced, CollapseCategories, NotBlueprintable)
@@ -134,7 +129,7 @@ enum class EWaveEditorFadeMode : uint8
 class FWaveTransformationTrimFade : public Audio::IWaveTransformation
 {
 public:
-	UE_API explicit FWaveTransformationTrimFade(const FFadeFunctionData& InFadeFunctions, double InStartTime, double InEndTime);
+	UE_API explicit FWaveTransformationTrimFade(const FFadeFunctionData& InFadeFunctions, double InStartTime, double InEndTime, float InStartFadeTime, float InStartFadeCurve, float InStartSCurveSharpness, float InEndFadeTime, float InEndFadeCurve, float InEndSCurveSharpness);
 	UE_API virtual void ProcessAudio(Audio::FWaveformTransformationWaveInfo& InOutWaveInfo) const override;
 
 	virtual constexpr Audio::ETransformationPriority FileChangeLengthPriority() const override { return Audio::ETransformationPriority::High; }
@@ -145,6 +140,14 @@ private:
 
 	double StartTime = 0.0;
 	double EndTime = 0.0;
+
+	float StartFadeTime = 0.f;
+	float StartFadeCurve = 0.f;
+	float StartSCurveSharpness = 0.f;
+
+	float EndFadeTime = 0.f;
+	float EndFadeCurve = 0.f;
+	float EndSCurveSharpness = 0.f;
 };
 
 UCLASS(MinimalAPI)
@@ -171,21 +174,21 @@ public:
 	 * Relative to StartTime
 	 * Individual fade properties have been deprecated. Use FadeFunctions instead.
 	 */ 
-	UPROPERTY(EditAnywhere, Category = "Fade", meta = (ClampMin = 0.0, DisplayName = "Fade-In Duration", EditCondition = "StartFadeTime != 0", EditConditionHides))
+	UPROPERTY(EditAnywhere, Category = "Fade", meta = (ClampMin = 0.0, DisplayName = "Fade-In Duration"))
 	float StartFadeTime = 0.f;
 	
 	UE_DEPRECATED(5.7, "Individual fade properties have been deprecated. Use FadeFunctions instead.")
 	/*
 	 * Individual fade properties have been deprecated. Use FadeFunctions instead.
 	 */
-	UPROPERTY(EditAnywhere, Category = "Fade", meta=(ClampMin = -0.1, ClampMax = 10.0, DisplayName = "Fade-In Curve", EditCondition = "StartFadeTime != 0", EditConditionHides))
+	UPROPERTY(EditAnywhere, Category = "Fade", meta=(ClampMin = -0.1, ClampMax = 10.0, DisplayName = "Fade-In Curve"))
 	float StartFadeCurve = 1.f;
 
 	UE_DEPRECATED(5.7, "Individual fade properties have been deprecated. Use FadeFunctions instead.")
 	/*
 	 * Individual fade properties have been deprecated. Use FadeFunctions instead.
 	 */
-	UPROPERTY(EditAnywhere, Category = "Fade", meta = (ClampMin = 0, ClampMax = 1.0, DisplayName = "S-Curve Sharpness", EditCondition = "StartFadeTime != 0 && StartFadeCurve < 0", EditConditionHides))
+	UPROPERTY(EditAnywhere, Category = "Fade", meta = (ClampMin = 0, ClampMax = 1.0, DisplayName = "S-Curve Sharpness", EditCondition = "StartFadeCurve < 0", EditConditionHides))
 	float StartSCurveSharpness = 0.75f;
 
 	UE_DEPRECATED(5.7, "Individual fade properties have been deprecated. Use FadeFunctions instead.")
@@ -193,21 +196,21 @@ public:
 	 * Relative to EndTime. Fade out starts at EndTime - EndFadeTime and ends at EndTime
 	 * Individual fade properties have been deprecated. Use FadeFunctions instead.
 	 */
-	UPROPERTY(EditAnywhere, Category = "Fade", meta = (ClampMin = 0.0, DisplayName = "Fade-Out Duration", EditCondition = "EndFadeTime != 0", EditConditionHides))
+	UPROPERTY(EditAnywhere, Category = "Fade", meta = (ClampMin = 0.0, DisplayName = "Fade-Out Duration"))
 	float EndFadeTime = 0.f;
 
 	UE_DEPRECATED(5.7, "Individual fade properties have been deprecated. Use FadeFunctions instead.")
 	/*
 	 * Individual fade properties have been deprecated. Use FadeFunctions instead.
 	 */
-	UPROPERTY(EditAnywhere, Category = "Fade", meta=(ClampMin = -0.1, ClampMax = 10.0, DisplayName = "Fade-Out Curve", EditCondition = "EndFadeTime != 0", EditConditionHides))
+	UPROPERTY(EditAnywhere, Category = "Fade", meta=(ClampMin = -0.1, ClampMax = 10.0, DisplayName = "Fade-Out Curve"))
 	float EndFadeCurve = 1.f;
 
 	UE_DEPRECATED(5.7, "Individual fade properties have been deprecated. Use FadeFunctions instead.")
 	/*
 	 * Individual fade properties have been deprecated. Use FadeFunctions instead.
 	 */
-	UPROPERTY(EditAnywhere, Category = "Fade", meta = (ClampMin = 0, ClampMax = 1.0, DisplayName = "S-Curve Sharpness", EditCondition = "EndFadeTime != 0 && EndFadeCurve < 0", EditConditionHides))
+	UPROPERTY(EditAnywhere, Category = "Fade", meta = (ClampMin = 0, ClampMax = 1.0, DisplayName = "S-Curve Sharpness", EditCondition = "EndFadeCurve < 0", EditConditionHides))
 	float EndSCurveSharpness = 0.75;
 
 	UE_DEPRECATED(5.7, "Individual fade properties have been deprecated. Use FadeFunctions instead.")
@@ -216,15 +219,13 @@ public:
 	UE_DEPRECATED(5.7, "Individual fade properties have been deprecated. Use FadeFunctions instead.")
 	static UE_API const float MaxFadeCurve;
 	
-	UE_DEPRECATED(5.7, "Individual fade properties have been deprecated. Use FadeFunctions instead.")
 	static UE_API const TMap<EWaveEditorFadeMode, float> FadeModeToCurveValueMap;
-
 	static UE_API const TMap<EWaveEditorFadeMode, TSubclassOf<UFadeFunction>> FadeModeToFadeFunctionMap;
 
 	UE_API virtual Audio::FTransformationPtr CreateTransformation() const override;
 
 	UE_API void UpdateConfiguration(FWaveTransformUObjectConfiguration& InOutConfiguration) override;
-	virtual constexpr Audio::ETransformationPriority GetTransformationPriority() const { return Audio::ETransformationPriority::High; }	
+	virtual constexpr Audio::ETransformationPriority GetTransformationPriority() const { return Audio::ETransformationPriority::High; }
 
 	UE_DEPRECATED(5.7, "Individual fade properties have been deprecated. Use FadeFunctions::FadeIn::GetFadeInCurveValue instead.")
 	static UE_API const double GetFadeInCurveValue(const float StartFadeCurve, const double FadeFraction, const float SCurveSharpness = 0);
@@ -233,10 +234,6 @@ public:
 	static UE_API const double GetFadeOutCurveValue(const float EndFadeCurve, const double FadeFraction, const float SCurveSharpness = 0);
 	
 private:	
-
-	UE_API virtual void PostInitProperties() override;
-	virtual void PostLoad() override;
-
 	UE_API void UpdateDurationProperties(const float InAvailableDuration);
 	float AvailableWaveformDuration = -1.f;
 };
diff --git a/Engine/Plugins/Editor/WaveformEditor/Source/WaveformTransformationsWidgets/Private/WaveformTransformationTrimFadeRenderer.cpp b/Engine/Plugins/Editor/WaveformEditor/Source/WaveformTransformationsWidgets/Private/WaveformTransformationTrimFadeRenderer.cpp
--- a/Engine/Plugins/Editor/WaveformEditor/Source/WaveformTransformationsWidgets/Private/WaveformTransformationTrimFadeRenderer.cpp
+++ b/Engine/Plugins/Editor/WaveformEditor/Source/WaveformTransformationsWidgets/Private/WaveformTransformationTrimFadeRenderer.cpp
@@ -158,76 +158,54 @@ int32 FWaveformTransformationTrimFadeRenderer::DrawFadeCurves(const FGeometry& A
 	return LayerId;
 }
 
+PRAGMA_DISABLE_DEPRECATION_WARNINGS
 void FWaveformTransformationTrimFadeRenderer::GenerateFadeCurves(const FGeometry& AllottedGeometry)
 {
 	check(StrongTrimFade);
-
-	check(StartTimeHandleX >= 0);
-	check(EndTimeHandleX >= 0);
+	const float FadeInTime = StrongTrimFade->FadeFunctions.FadeIn ? StrongTrimFade->FadeFunctions.FadeIn->Duration : StrongTrimFade->StartFadeTime;
+	const float FadeInFrames = FadeInTime * TransformationWaveInfo.SampleRate;
+	const uint32 FadeInPixelLength = FadeInFrames * PixelsPerFrame;
 	FadeInStartX = FMath::RoundToInt32(StartTimeHandleX);
-	FadeOutEndX = FMath::RoundToInt32(EndTimeHandleX);
+	FadeInEndX = FMath::RoundToInt32(FMath::Clamp(StartTimeHandleX + FadeInPixelLength, StartTimeHandleX, EndTimeHandleX));
+	
+	const uint32 DisplayedFadeInPixelLength = FadeInEndX - FadeInStartX;
+	FadeInCurvePoints.SetNumUninitialized(DisplayedFadeInPixelLength);
 
-	if (StrongTrimFade->FadeFunctions.FadeIn && StrongTrimFade->FadeFunctions.FadeIn->Duration > 0.f)
+	for (uint32 Pixel = 0; Pixel < DisplayedFadeInPixelLength; ++Pixel)
 	{
-		const float FadeInTime = StrongTrimFade->FadeFunctions.FadeIn->Duration;
-		const float FadeInFrames = FadeInTime * TransformationWaveInfo.SampleRate;
-		const float FadeInPixelLength = FadeInFrames * PixelsPerFrame;
-		check(FadeInPixelLength > 0.f);
-		check(EndTimeHandleX >= 0);
-		
-		FadeInEndX = FMath::RoundToInt32(FMath::Clamp(StartTimeHandleX + FadeInPixelLength, StartTimeHandleX, EndTimeHandleX));
-
-		const uint32 DisplayedFadeInPixelLength = FadeInEndX - FadeInStartX;
-		FadeInCurvePoints.SetNumUninitialized(DisplayedFadeInPixelLength);
+		const double FadeFraction = (float)Pixel / FadeInPixelLength;
+		const double CurveFunction = StrongTrimFade->FadeFunctions.FadeIn ? StrongTrimFade->FadeFunctions.FadeIn->GetFadeInCurveValue(FadeFraction)
+			: UWaveformTransformationTrimFade::GetFadeInCurveValue(StrongTrimFade->StartFadeCurve, FadeFraction, StrongTrimFade->StartSCurveSharpness);
 
-		for (uint32 Pixel = 0; Pixel < DisplayedFadeInPixelLength; ++Pixel)
-		{
-			const double FadeFraction = static_cast<float>(Pixel) / FadeInPixelLength;
-			const double CurveFunction = StrongTrimFade->FadeFunctions.FadeIn->GetFadeInCurveValue(FadeFraction);
+		const double CurveValue = Pixel != FadeInPixelLength - 1 ? 1.f - CurveFunction : 0.f;
 
-			const double CurveValue = Pixel != FadeInPixelLength - 1 ? 1.f - CurveFunction : 0.f;
-
-			const uint32 XCoordinate = Pixel + FadeInStartX;
-			FadeInCurvePoints[Pixel] = FVector2D(XCoordinate, CurveValue * AllottedGeometry.Size.Y);
-		}
-	}
-	else
-	{
-		FadeInEndX = FMath::RoundToInt32(FMath::Clamp(StartTimeHandleX, StartTimeHandleX, EndTimeHandleX));
-		FadeInCurvePoints.SetNumUninitialized(0);
+		const uint32 XCoordinate = Pixel + FadeInStartX;
+		FadeInCurvePoints[Pixel] = FVector2D(XCoordinate, CurveValue * AllottedGeometry.Size.Y);
 	}
 
-	if(StrongTrimFade->FadeFunctions.FadeOut && StrongTrimFade->FadeFunctions.FadeOut->Duration > 0.f)
-	{
-		const float FadeOutTime = StrongTrimFade->FadeFunctions.FadeOut->Duration;
-		const float FadeOutFrames = FadeOutTime * TransformationWaveInfo.SampleRate;
-		const float FadeOutPixelLength = FadeOutFrames * PixelsPerFrame;
-		check(FadeOutPixelLength > 0.f);
-
-		check(StartTimeHandleX >= 0);
-		FadeOutStartX = FMath::RoundToInt32(FMath::Clamp(EndTimeHandleX - FadeOutPixelLength, StartTimeHandleX, EndTimeHandleX));
-
-		const uint32 DisplayedFadeOutPixelLength = FadeOutEndX - FadeOutStartX;
-		FadeOutCurvePoints.SetNumUninitialized(DisplayedFadeOutPixelLength);
-		const uint32 FadeOutPixelOffset = FMath::RoundToInt32(FMath::Max(FadeOutPixelLength - static_cast<float>(DisplayedFadeOutPixelLength), 0.f));
+	const float FadeOutTime = StrongTrimFade->FadeFunctions.FadeOut ? StrongTrimFade->FadeFunctions.FadeOut->Duration : StrongTrimFade->EndFadeTime;
+	const float FadeOutFrames = FadeOutTime * TransformationWaveInfo.SampleRate;
+	const float FadeOutPixelLength = FadeOutFrames * PixelsPerFrame;
+	FadeOutStartX = FMath::RoundToInt32(FMath::Clamp(EndTimeHandleX - FadeOutPixelLength, StartTimeHandleX, EndTimeHandleX));
+	FadeOutEndX = FMath::RoundToInt32(EndTimeHandleX);
+	
+	const uint32 DisplayedFadeOutPixelLength = FadeOutEndX - FadeOutStartX;
+	FadeOutCurvePoints.SetNumUninitialized(DisplayedFadeOutPixelLength);
+	const uint32 FadeOutPixelOffset = FadeOutPixelLength - DisplayedFadeOutPixelLength;
 
-		for (uint32 Pixel = 0; Pixel < DisplayedFadeOutPixelLength; ++Pixel)
-		{
-			const double FadeFraction = static_cast<float>(Pixel + FadeOutPixelOffset) / FadeOutPixelLength;
-			const double CurveFunction = StrongTrimFade->FadeFunctions.FadeOut->GetFadeOutCurveValue(FadeFraction);
+	for (uint32 Pixel = 0; Pixel < DisplayedFadeOutPixelLength; ++Pixel)
+	{
+		const double FadeFraction = (float)(Pixel + FadeOutPixelOffset) / FadeOutPixelLength;
+		const double CurveFunction = StrongTrimFade->FadeFunctions.FadeOut ? StrongTrimFade->FadeFunctions.FadeOut->GetFadeOutCurveValue(FadeFraction)
+			: UWaveformTransformationTrimFade::GetFadeOutCurveValue(StrongTrimFade->EndFadeCurve, FadeFraction, StrongTrimFade->EndSCurveSharpness);
 
-			const double CurveValue = Pixel != FadeOutPixelLength - 1 ? 1.f - CurveFunction : 1.f;
+		const double CurveValue = Pixel != FadeOutPixelLength - 1 ? 1.f - CurveFunction : 1.f;
 
-			const uint32 XCoordinate = Pixel + FadeOutStartX;
-			FadeOutCurvePoints[Pixel] = FVector2D(XCoordinate, CurveValue * AllottedGeometry.Size.Y);
-		}
-	}
-	else
-	{
-		FadeOutStartX = FMath::RoundToInt32(FMath::Clamp(EndTimeHandleX, StartTimeHandleX, EndTimeHandleX));
-		FadeOutCurvePoints.SetNumUninitialized(0);
+		const uint32 XCoordinate = Pixel + FadeOutStartX;
+		FadeOutCurvePoints[Pixel] = FVector2D(XCoordinate, CurveValue * AllottedGeometry.Size.Y);
 	}
 }
+PRAGMA_ENABLE_DEPRECATION_WARNINGS
 
 FCursorReply FWaveformTransformationTrimFadeRenderer::OnCursorQuery(const FGeometry& MyGeometry, const FPointerEvent& CursorEvent) const
 {
@@ -264,24 +242,46 @@ FReply FWaveformTransformationTrimFadeRenderer::OnMouseWheel(SWidget& OwnerWidge
 	check(StrongTrimFade);
 	const FVector2D LocalCursorPosition = GetLocalCursorPosition(MouseEvent, MyGeometry);
 
-	if (IsCursorInFadeInInteractionRange(LocalCursorPosition, MyGeometry) && StrongTrimFade->FadeFunctions.FadeIn && StrongTrimFade->FadeFunctions.FadeIn.IsA<UFadeCurveFunction>())
-	{
+	if (IsCursorInFadeInInteractionRange(LocalCursorPosition, MyGeometry))
+	{		
 		const float FadeCurveDelta = MouseEvent.GetWheelDelta() * MouseWheelStep;
 		FText OutPropertyName = FText::FromName(NAME_None);
+		
+		if (StrongTrimFade->FadeFunctions.FadeIn && StrongTrimFade->FadeFunctions.FadeIn.IsA<UFadeCurveFunction>())
+		{
+			TObjectPtr<UFadeCurveFunction> FadeIn = CastChecked<UFadeCurveFunction>(StrongTrimFade->FadeFunctions.FadeIn);
+			OutPropertyName = FadeIn->GetFadeCurvePropertyName();
+			
+			if (BeginTransaction(TEXT("PropertyEditor"), FText::Format(NSLOCTEXT("PropertyEditor", "SetNumericPropertyTransaction", "Edit {0}"), OutPropertyName), StrongTrimFade.Get()) == INDEX_NONE)
+			{
+				UE_LOG(LogWaveformTransformation, Warning, TEXT("Begin TrimFadeTransformation OnMouseWheel Transaction Failed"));
+			}
 
-		TObjectPtr<UFadeCurveFunction> FadeIn = CastChecked<UFadeCurveFunction>(StrongTrimFade->FadeFunctions.FadeIn);
-		OutPropertyName = FadeIn->GetFadeCurvePropertyName();
+			StrongTrimFade->Modify();
+			StrongTrimFade->FadeFunctions.FadeIn->Modify();
 
-		if (BeginTransaction(TEXT("PropertyEditor"), FText::Format(NSLOCTEXT("PropertyEditor", "SetNumericPropertyTransaction", "Edit {0}"), OutPropertyName), StrongTrimFade.Get()) == INDEX_NONE)
+			const float FadeCurveValue = FadeIn->GetFadeCurve() + FadeCurveDelta;
+			FadeIn->SetFadeCurve(FadeCurveValue);
+		}
+		else
 		{
-			UE_LOG(LogWaveformTransformation, Warning, TEXT("Begin TrimFadeTransformation OnMouseWheel Transaction Failed"));
+			PRAGMA_DISABLE_DEPRECATION_WARNINGS
+			OutPropertyName = FText::FromName(GET_MEMBER_NAME_CHECKED(UWaveformTransformationTrimFade, StartFadeCurve));
+			if (BeginTransaction(TEXT("PropertyEditor"), FText::Format(NSLOCTEXT("PropertyEditor", "SetNumericPropertyTransaction", "Edit {0}"), OutPropertyName), StrongTrimFade.Get()) == INDEX_NONE)
+			{
+				UE_LOG(LogWaveformTransformation, Warning, TEXT("Begin TrimFadeTransformation OnMouseWheel Transaction Failed"));
+			}
+
+			StrongTrimFade->Modify();
+
+			const float FadeCurveValue = FMath::Clamp(StrongTrimFade->StartFadeCurve + FadeCurveDelta, 0.f, 10.f);
+			StrongTrimFade->StartFadeCurve = FadeCurveValue;
+			PRAGMA_ENABLE_DEPRECATION_WARNINGS
 		}
 
-		StrongTrimFade->Modify();
-		StrongTrimFade->FadeFunctions.FadeIn->Modify();
+		EPropertyValueSetFlags::Type PropertyFlags = EPropertyValueSetFlags::InteractiveChange;
 
-		const float FadeCurveValue = FadeIn->GetFadeCurve() + FadeCurveDelta;
-		FadeIn->SetFadeCurve(FadeCurveValue);
+		StrongTrimFade->PostEditChange();
 
 		EndTransaction();
 		StrongTrimFade->OnTransformationChanged.ExecuteIfBound(true);
@@ -289,24 +289,46 @@ FReply FWaveformTransformationTrimFadeRenderer::OnMouseWheel(SWidget& OwnerWidge
 		return FReply::Handled();
 	}
 
-	if (IsCursorInFadeOutInteractionRange(LocalCursorPosition, MyGeometry) && StrongTrimFade->FadeFunctions.FadeOut && StrongTrimFade->FadeFunctions.FadeOut.IsA<UFadeCurveFunction>())
+	if (IsCursorInFadeOutInteractionRange(LocalCursorPosition, MyGeometry))
 	{
 		const float FadeCurveDelta = MouseEvent.GetWheelDelta() * MouseWheelStep;
 		FText OutPropertyName = FText::FromName(NAME_None);
 
-		TObjectPtr<UFadeCurveFunction> FadeOut = CastChecked<UFadeCurveFunction>(StrongTrimFade->FadeFunctions.FadeOut);
-		OutPropertyName = FadeOut->GetFadeCurvePropertyName();
-		if (BeginTransaction(TEXT("PropertyEditor"), FText::Format(NSLOCTEXT("PropertyEditor", "SetNumericPropertyTransaction", "Edit {0}"), OutPropertyName), StrongTrimFade.Get()) == INDEX_NONE)
+		if (StrongTrimFade->FadeFunctions.FadeOut && StrongTrimFade->FadeFunctions.FadeOut.IsA<UFadeCurveFunction>())
 		{
-			UE_LOG(LogWaveformTransformation, Warning, TEXT("Begin TrimFadeTransformation OnMouseWheel Transaction Failed"));
+			TObjectPtr<UFadeCurveFunction> FadeOut = CastChecked<UFadeCurveFunction>(StrongTrimFade->FadeFunctions.FadeOut);
+			OutPropertyName = FadeOut->GetFadeCurvePropertyName();
+			if (BeginTransaction(TEXT("PropertyEditor"), FText::Format(NSLOCTEXT("PropertyEditor", "SetNumericPropertyTransaction", "Edit {0}"), OutPropertyName), StrongTrimFade.Get()) == INDEX_NONE)
+			{
+				UE_LOG(LogWaveformTransformation, Warning, TEXT("Begin TrimFadeTransformation OnMouseWheel Transaction Failed"));
+			}
+
+			StrongTrimFade->Modify();
+			StrongTrimFade->FadeFunctions.FadeOut->Modify();
+
+			const float FadeCurveValue = FadeOut->GetFadeCurve() + FadeCurveDelta;
+			FadeOut->SetFadeCurve(FadeCurveValue);
 		}
+		else
+		{
+			PRAGMA_DISABLE_DEPRECATION_WARNINGS
+			OutPropertyName = FText::FromName(GET_MEMBER_NAME_CHECKED(UWaveformTransformationTrimFade, EndFadeCurve));
+			if (BeginTransaction(TEXT("PropertyEditor"), FText::Format(NSLOCTEXT("PropertyEditor", "SetNumericPropertyTransaction", "Edit {0}"), OutPropertyName), StrongTrimFade.Get()) == INDEX_NONE)
+			{
+				UE_LOG(LogWaveformTransformation, Warning, TEXT("Begin TrimFadeTransformation OnMouseWheel Transaction Failed"));
+			}
 
-		StrongTrimFade->Modify();
-		StrongTrimFade->FadeFunctions.FadeOut->Modify();
+			StrongTrimFade->Modify();
 
-		const float FadeCurveValue = FadeOut->GetFadeCurve() + FadeCurveDelta;
-		FadeOut->SetFadeCurve(FadeCurveValue);
+			const float FadeCurveValue = FMath::Clamp(StrongTrimFade->EndFadeCurve + FadeCurveDelta, 0.f, 10.f);
+			StrongTrimFade->EndFadeCurve = FadeCurveValue;
+			PRAGMA_ENABLE_DEPRECATION_WARNINGS
+		}
+		
+		EPropertyValueSetFlags::Type PropertyFlags = EPropertyValueSetFlags::InteractiveChange;
 
+		StrongTrimFade->PostEditChange();
+		
 		EndTransaction();
 		StrongTrimFade->OnTransformationChanged.ExecuteIfBound(true);
 
@@ -386,18 +408,7 @@ FReply FWaveformTransformationTrimFadeRenderer::OnMouseButtonDown(SWidget& Owner
 		{
 			UE_LOG(LogWaveformTransformation, Warning, TEXT("Begin TrimFadeTransformation OnMouseButtonDown Transaction Failed"));
 		}
-
-		if (TrimFadeInteractionType == ETrimFadeInteractionType::ScrubbingFadeIn && StrongTrimFade->FadeFunctions.FadeIn)
-		{
-			StrongTrimFade->FadeFunctions.FadeIn->Modify();
-		}
-		else if (TrimFadeInteractionType == ETrimFadeInteractionType::ScrubbingFadeOut && StrongTrimFade->FadeFunctions.FadeOut)
-		{
-			StrongTrimFade->FadeFunctions.FadeOut->Modify();
-		}
-		
 		StrongTrimFade->Modify();
-
 		return FReply::Handled().CaptureMouse(OwnerWidget.AsShared()).PreventThrottling();
 	}
 
@@ -411,9 +422,6 @@ FReply FWaveformTransformationTrimFadeRenderer::OnMouseMove(SWidget& OwnerWidget
 	if (MouseEvent.IsMouseButtonDown(EKeys::LeftMouseButton) && TrimFadeInteractionType != ETrimFadeInteractionType::None)
 	{
 		SetPropertyValueDependingOnInteractionType(MouseEvent, MyGeometry);
-
-		StrongTrimFade->PostEditChange();
-
 		StrongTrimFade->OnTransformationRenderChanged.ExecuteIfBound();
 
 		return FReply::Handled().CaptureMouse(OwnerWidget.AsShared());
@@ -427,12 +435,11 @@ FReply FWaveformTransformationTrimFadeRenderer::OnMouseButtonUp(SWidget& OwnerWi
 	check(StrongTrimFade);
 	if (TrimFadeInteractionType != ETrimFadeInteractionType::None)
 	{
+		StrongTrimFade->Modify();
+
 		SetPropertyValueDependingOnInteractionType(MouseEvent, MyGeometry);
 
 		EndTransaction();
-
-		StrongTrimFade->PostEditChange();
-
 		StrongTrimFade->OnTransformationChanged.ExecuteIfBound(true);
 
 		TrimFadeInteractionType = ETrimFadeInteractionType::None;
@@ -463,21 +470,39 @@ void FWaveformTransformationTrimFadeRenderer::SetPropertyValueDependingOnInterac
 	case FWaveformTransformationTrimFadeRenderer::ETrimFadeInteractionType::ScrubbingFadeIn:
 		{
 			float StartFadeTimeValue = 0;
+
 			if (StrongTrimFade->FadeFunctions.FadeIn)
 			{
+				StrongTrimFade->FadeFunctions.FadeIn->Modify();
 				StartFadeTimeValue = FMath::Clamp(SelectedTime - StrongTrimFade->StartTime, 0.f, TNumericLimits<float>().Max());
 				StrongTrimFade->FadeFunctions.FadeIn->Duration = StartFadeTimeValue;
 			}
+			else
+			{
+				PRAGMA_DISABLE_DEPRECATION_WARNINGS
+				StartFadeTimeValue = FMath::Clamp(SelectedTime - StrongTrimFade->StartTime, 0.f, TNumericLimits<float>().Max());
+				StrongTrimFade->StartFadeTime = StartFadeTimeValue;
+				PRAGMA_ENABLE_DEPRECATION_WARNINGS
+			}
 		}
 		break;
 	case FWaveformTransformationTrimFadeRenderer::ETrimFadeInteractionType::ScrubbingFadeOut:
 		{	
 			float EndFadeTimeValue = 0;
+
 			if (StrongTrimFade->FadeFunctions.FadeOut)
 			{
+				StrongTrimFade->FadeFunctions.FadeOut->Modify();
 				EndFadeTimeValue = FMath::Clamp(StrongTrimFade->EndTime - SelectedTime, 0.f, TNumericLimits<float>().Max());
 				StrongTrimFade->FadeFunctions.FadeOut->Duration = EndFadeTimeValue;
 			}
+			else
+			{
+				PRAGMA_DISABLE_DEPRECATION_WARNINGS
+				EndFadeTimeValue = FMath::Clamp(StrongTrimFade->EndTime - SelectedTime, 0.f, TNumericLimits<float>().Max());
+				StrongTrimFade->EndFadeTime = EndFadeTimeValue;
+				PRAGMA_ENABLE_DEPRECATION_WARNINGS
+			}
 		}
 		break;
 	case FWaveformTransformationTrimFadeRenderer::ETrimFadeInteractionType::RightClickFadeIn:
@@ -519,7 +544,7 @@ void FWaveformTransformationTrimFadeRenderer::ShowSelectFadeModeMenuAtCursor(con
 	FadeModeMenuWindow = MenuWindow;
 	TWeakObjectPtr<UWaveformTransformationTrimFade> InWeakTrimFade = StrongTrimFade.Get();
 
-	for (const TPair<EWaveEditorFadeMode, TSubclassOf<UFadeFunction>>& FadeOptionPair : UWaveformTransformationTrimFade::FadeModeToFadeFunctionMap)
+	for (const TPair<EWaveEditorFadeMode, float>& FadeOptionPair : UWaveformTransformationTrimFade::FadeModeToCurveValueMap)
 	{
 		MenuContent->AddSlot()
 			.Padding(5)
@@ -537,24 +562,24 @@ void FWaveformTransformationTrimFadeRenderer::ShowSelectFadeModeMenuAtCursor(con
 
 							PinnedTrimFade->Modify();
 
+							PRAGMA_DISABLE_DEPRECATION_WARNINGS
 							if (InteractionType == FWaveformTransformationTrimFadeRenderer::ETrimFadeInteractionType::RightClickFadeIn)
 							{
 								if (PinnedTrimFade->FadeFunctions.FadeIn)
 								{
-									const float Duration = PinnedTrimFade->FadeFunctions.FadeIn->Duration;
-									PinnedTrimFade->FadeFunctions.FadeIn = NewObject<UFadeFunction>(PinnedTrimFade.Get(), FadeOptionPair.Value.Get(), NAME_None, RF_Transactional);
-									PinnedTrimFade->FadeFunctions.FadeIn->Duration =  Duration;
+									PinnedTrimFade->FadeFunctions.FadeIn = NewObject<UFadeFunction>(PinnedTrimFade.Get(), UWaveformTransformationTrimFade::FadeModeToFadeFunctionMap[FadeOptionPair.Key].Get());
 								}
+								PinnedTrimFade->StartFadeCurve = FadeOptionPair.Value;
 							}
 							else if (InteractionType == FWaveformTransformationTrimFadeRenderer::ETrimFadeInteractionType::RightClickFadeOut)
 							{
 								if (PinnedTrimFade->FadeFunctions.FadeOut)
 								{
-									const float Duration = PinnedTrimFade->FadeFunctions.FadeOut->Duration;
-									PinnedTrimFade->FadeFunctions.FadeOut = NewObject<UFadeFunction>(PinnedTrimFade.Get(), FadeOptionPair.Value.Get(), NAME_None, RF_Transactional);
-									PinnedTrimFade->FadeFunctions.FadeOut->Duration = Duration;
+									PinnedTrimFade->FadeFunctions.FadeOut = NewObject<UFadeFunction>(PinnedTrimFade.Get(), UWaveformTransformationTrimFade::FadeModeToFadeFunctionMap[FadeOptionPair.Key].Get());
 								}
+								PinnedTrimFade->EndFadeCurve = FadeOptionPair.Value;
 							}
+							PRAGMA_ENABLE_DEPRECATION_WARNINGS
 
 							if (GEditor && GEditor->Trans)
 							{
@@ -615,29 +640,20 @@ void FWaveformTransformationTrimFadeRenderer::ShowSelectFadeModeMenuAtCursor(con
 			return false; // return false to only execute once
 		}));
 
-	TSharedPtr<FDelegateHandle> ApplicationActivationStateHandlePtr = MakeShared<FDelegateHandle>();
-	*ApplicationActivationStateHandlePtr = FSlateApplication::Get().OnApplicationActivationStateChanged().AddLambda([InFadeModeMenuWindow = FadeModeMenuWindow, ApplicationActivationStateHandlePtr](bool isActive)
+	ApplicationActivationStateHandle = FSlateApplication::Get().OnApplicationActivationStateChanged().AddLambda([InFadeModeMenuWindow = FadeModeMenuWindow, InApplicationActivationStateHandle = ApplicationActivationStateHandle](bool isActive)
 		{
 			if (!isActive)
 			{
 				if (TSharedPtr<SWindow> LockedFadeModeMenuWindow = InFadeModeMenuWindow.Pin())
 				{
-					FSlateApplication::Get().RequestDestroyWindow(LockedFadeModeMenuWindow.ToSharedRef());
-					FSlateApplication::Get().OnApplicationActivationStateChanged().Remove(*ApplicationActivationStateHandlePtr);
+					FSlateApplication::Get().RequestDestroyWindow(LockedFadeModeMenuWindow.ToSharedRef());	
+					FSlateApplication::Get().OnApplicationActivationStateChanged().Remove(InApplicationActivationStateHandle);
 				}
 			}
 		});
 
-	if (ApplicationActivationStateHandle.IsValid())
-	{
-		FSlateApplication::Get().OnApplicationActivationStateChanged().Remove(ApplicationActivationStateHandle);
-	}
-
-	ApplicationActivationStateHandle = *ApplicationActivationStateHandlePtr;
-
 	// If focus is lost on the popup, destroy it to prevent popups hanging around
-	TSharedPtr<FDelegateHandle> PopupHandlePtr = MakeShared<FDelegateHandle>();
-	*PopupHandlePtr = FSlateApplication::Get().OnFocusChanging().AddLambda([InFadeModeMenuWindow = FadeModeMenuWindow, PopupHandlePtr](const FFocusEvent& FocusEvent, const FWeakWidgetPath& WeakWidgetPath
+	PopupHandle = FSlateApplication::Get().OnFocusChanging().AddLambda([InFadeModeMenuWindow = FadeModeMenuWindow, InPopupHandle = PopupHandle](const FFocusEvent& FocusEvent, const FWeakWidgetPath& WeakWidgetPath
 		, const TSharedPtr<SWidget>& OldWidget, const FWidgetPath& WidgetPath, const TSharedPtr<SWidget>& NewWidget)
 		{
 			if (InFadeModeMenuWindow != nullptr && InFadeModeMenuWindow.IsValid())
@@ -647,18 +663,11 @@ void FWaveformTransformationTrimFadeRenderer::ShowSelectFadeModeMenuAtCursor(con
 					if (OldWidget && !OldWidget->IsHovered() && LockedFadeModeMenuWindow == OldWidget)
 					{
 						FSlateApplication::Get().RequestDestroyWindow(LockedFadeModeMenuWindow.ToSharedRef());
-						FSlateApplication::Get().OnFocusChanging().Remove(*PopupHandlePtr);
+						FSlateApplication::Get().OnFocusChanging().Remove(InPopupHandle);
 					}
 				}
 			}
 		});
-
-	if (PopupHandle.IsValid())
-	{
-		FSlateApplication::Get().OnFocusChanging().Remove(PopupHandle);
-	}
-
-	PopupHandle = *PopupHandlePtr;
 }
 
 const FSlateIcon FWaveformTransformationTrimFadeRenderer::GetFadeModeIcon(const EWaveEditorFadeMode& FadeMode)
@@ -772,6 +781,12 @@ FText FWaveformTransformationTrimFadeRenderer::GetPropertyEditedByCurrentInterac
 				FText::FromName(GET_MEMBER_NAME_CHECKED(FFadeFunctionData, FadeIn)),
 				FText::FromName(GET_MEMBER_NAME_CHECKED(UFadeFunction, Duration)));
 		}
+		else
+		{
+			PRAGMA_DISABLE_DEPRECATION_WARNINGS
+			OutPropertyName = FText::FromName(GET_MEMBER_NAME_CHECKED(UWaveformTransformationTrimFade, StartFadeTime));
+			PRAGMA_ENABLE_DEPRECATION_WARNINGS
+		}
 		break;
 	case FWaveformTransformationTrimFadeRenderer::ETrimFadeInteractionType::ScrubbingFadeOut:
 		if (StrongTrimFade->FadeFunctions.FadeOut)
@@ -780,6 +795,12 @@ FText FWaveformTransformationTrimFadeRenderer::GetPropertyEditedByCurrentInterac
 				FText::FromName(GET_MEMBER_NAME_CHECKED(FFadeFunctionData, FadeOut)),
 				FText::FromName(GET_MEMBER_NAME_CHECKED(UFadeFunction, Duration)));
 		}
+		else
+		{
+			PRAGMA_DISABLE_DEPRECATION_WARNINGS
+			OutPropertyName = FText::FromName(GET_MEMBER_NAME_CHECKED(UWaveformTransformationTrimFade, EndFadeTime));
+			PRAGMA_ENABLE_DEPRECATION_WARNINGS
+		}
 		break;
 	default:
 		break;
