[Wavform] Display more SoundWave metadata via content browser.
Merge to release 38

diff --git a/Engine/Plugins/Editor/WaveformEditor/Source/WaveformEditor/Private/WaveformEditor.cpp b/Engine/Plugins/Editor/WaveformEditor/Source/WaveformEditor/Private/WaveformEditor.cpp
--- a/Engine/Plugins/Editor/WaveformEditor/Source/WaveformEditor/Private/WaveformEditor.cpp
+++ b/Engine/Plugins/Editor/WaveformEditor/Source/WaveformEditor/Private/WaveformEditor.cpp
@@ -20,6 +20,7 @@
 #include "ToolMenus.h"
 #include "TransformedWaveformView.h"
 #include "TransformedWaveformViewFactory.h"
+#include "WaveformAudioAnalysisFunctions.h"
 #include "WaveformEditorCommands.h"
 #include "WaveformEditorDetailsCustomization.h"
 #include "WaveformEditorLog.h"
@@ -189,6 +190,9 @@ bool FWaveformEditor::Init(const EToolkitMode::Type Mode, const TSharedPtr<ITool
 					});
 			}
 		}
+#if WITH_EDITOR
+		UpdateSoundWaveProperties();
+#endif // WITH_EDITOR
 	}
 
 	FSlateApplication::Get().SetKeyboardFocus(TransformationsDetails);
@@ -708,6 +712,12 @@ void FWaveformEditor::UpdateTransformations(bool bMarkDirty)
 	{
 		SoundWave->GetPackage()->ClearDirtyFlag();
 	}
+	else
+	{
+#if WITH_EDITOR
+		UpdateSoundWaveProperties();
+#endif //WITH_EDITOR
+	}
 
 	if (bIsPlaying)
 	{
@@ -1814,6 +1824,33 @@ void FWaveformEditor::UpdateTransportState()
 	}
 }
 
+#if WITH_EDITOR
+void FWaveformEditor::UpdateSoundWaveProperties()
+{
+	check(SoundWave);
+	check(WaveformView.DataProvider);
+
+	const TArray<float> TransformedPCMData = WaveformView.DataProvider->GetTransformedPCMData();
+	const uint32 SampleRate = TransformationChainConfig.SampleRate != 0 ?
+		TransformationChainConfig.SampleRate : SoundWave->GetImportedSampleRate();
+
+	if (TransformedPCMData.IsEmpty() || SampleRate == 0)
+	{
+		return;
+	}
+
+	const int64 StartSample = FMath::RoundToInt64(TransformationChainConfig.StartTime * SampleRate * TransformationChainConfig.NumChannels);
+	const int64 EndSample = FMath::RoundToInt64(TransformationChainConfig.EndTime * SampleRate * TransformationChainConfig.NumChannels);
+
+	const Audio::FAlignedFloatBuffer AudioBuffer(TransformedPCMData.GetData() + StartSample, FMath::Clamp(EndSample - StartSample, 0, TransformedPCMData.Num() - StartSample));
+
+	const float SamplePeakDB = WaveformAudioAnalysis::GetPeakSampleValue(AudioBuffer);
+	const float LUFS = WaveformAudioAnalysis::GetLUFS(AudioBuffer, SampleRate, TransformationChainConfig.NumChannels);
+
+	SoundWave->SetLoudnessValues(LUFS, SamplePeakDB);
+}
+#endif //WITH_EDITOR
+
 void FWaveformEditor::PostUndo(bool bSuccess)
 {
 	if (bSuccess)
diff --git a/Engine/Plugins/Editor/WaveformEditor/Source/WaveformEditor/Private/WaveformEditorModule.cpp b/Engine/Plugins/Editor/WaveformEditor/Source/WaveformEditor/Private/WaveformEditorModule.cpp
--- a/Engine/Plugins/Editor/WaveformEditor/Source/WaveformEditor/Private/WaveformEditorModule.cpp
+++ b/Engine/Plugins/Editor/WaveformEditor/Source/WaveformEditor/Private/WaveformEditorModule.cpp
@@ -2,10 +2,13 @@
 
 #include "WaveformEditorModule.h"
 
+#include "DSP/FloatArrayMath.h"
 #include "Editor.h"
 #include "Modules/ModuleManager.h"
 #include "PropertyEditorModule.h"
+#include "Sound/SoundWave.h"
 #include "TransformedWaveformViewFactory.h"
+#include "WaveformAudioAnalysisFunctions.h"
 #include "WaveformEditorCommands.h"
 #include "WaveformEditorInstantiator.h"
 #include "WaveformEditorLog.h"
@@ -21,11 +24,59 @@ void FWaveformEditorModule::StartupModule()
 	WaveformEditorInstantiator = MakeShared<FWaveformEditorInstantiator>();
 	RegisterContentBrowserExtensions(WaveformEditorInstantiator.Get());
 	WaveformEditorInstantiator->RegisterAsSoundwaveEditor();
+
+	OnPostEngineInitHandle = FCoreDelegates::OnPostEngineInit.AddLambda([this]()
+		{
+			if (GEditor)
+			{
+				if (TObjectPtr<UImportSubsystem> ImportSubsystem = GEditor->GetEditorSubsystem<UImportSubsystem>())
+				{
+					OnAssetPostImportHandle = ImportSubsystem->OnAssetPostImport.AddLambda([](UFactory* Factory, UObject* Object)
+						{
+							if (Object && Object->IsA<USoundWave>())
+							{
+								TObjectPtr<USoundWave> SoundWave = CastChecked<USoundWave>(Object);
+
+								TArray<uint8> ImportedRawPCMData;
+								uint32 BufferSampleRate = 0;
+								uint16 ImportedNumChannels = 0;
+
+								Audio::FAlignedFloatBuffer AudioBuffer;
+
+								if (!SoundWave->GetImportedSoundWaveData(ImportedRawPCMData, BufferSampleRate, ImportedNumChannels) || BufferSampleRate == 0)
+								{
+									UE_LOG(LogWaveformEditor, Warning, TEXT("Soundwave loudness analyzation failed on import."))
+										return;
+								}
+
+								const int64 NumWaveformSamples = ImportedRawPCMData.Num() * sizeof(uint8) / sizeof(int16);
+								AudioBuffer.SetNumUninitialized(NumWaveformSamples);
+								Audio::ArrayPcm16ToFloat(MakeArrayView((int16*)ImportedRawPCMData.GetData(), NumWaveformSamples), AudioBuffer);
+
+								const float SamplePeakDB = WaveformAudioAnalysis::GetPeakSampleValue(AudioBuffer);
+								const float LUFS = WaveformAudioAnalysis::GetLUFS(AudioBuffer, BufferSampleRate, ImportedNumChannels);
+
+								SoundWave->SetLoudnessValues(LUFS, SamplePeakDB);
+							}
+						});
+				}
+			}
+		});
 }
 
 void FWaveformEditorModule::ShutdownModule()
 {
 	FWaveformEditorCommands::Unregister();
+
+	FCoreDelegates::OnPostEngineInit.Remove(OnPostEngineInitHandle);
+
+	if (GEditor)
+	{
+		if (TObjectPtr<UImportSubsystem> ImportSubsystem = GEditor->GetEditorSubsystem<UImportSubsystem>())
+		{
+			ImportSubsystem->OnAssetPostImport.Remove(OnAssetPostImportHandle);
+		}
+	}
 }
 
 void FWaveformEditorModule::RegisterContentBrowserExtensions(IWaveformEditorInstantiator* Instantiator)
diff --git a/Engine/Plugins/Editor/WaveformEditor/Source/WaveformEditor/Public/WaveformEditor.h b/Engine/Plugins/Editor/WaveformEditor/Source/WaveformEditor/Public/WaveformEditor.h
--- a/Engine/Plugins/Editor/WaveformEditor/Source/WaveformEditor/Public/WaveformEditor.h
+++ b/Engine/Plugins/Editor/WaveformEditor/Source/WaveformEditor/Public/WaveformEditor.h
@@ -125,12 +125,18 @@ private:
 	UE_API bool CanFadeOut();
 	UE_API void CreateMarker(bool bIsLoopRegion);
 	UE_API void DeleteMarker();
-	UE_API void SkipToNextMarker(); // Skips the playhead to the next marker
+	// Skips the playhead to the next marker
+	UE_API void SkipToNextMarker();
 	UE_API void RegenerateTransformations();
 	UE_API void UpdateTransformations(bool bMarkDirty = true);
 	UE_API void UpdateRenderElements();
 	UE_API void UpdateTransportState();
 
+#if WITH_EDITOR
+	// Update SoundWave LUFS and SamplePeak values
+	UE_API void UpdateSoundWaveProperties();
+#endif //WITH_EDITOR
+
 	UE_API void ModifyMarkerLoopRegion(ELoopModificationControls Modification);
 	UE_API void CycleMarkerLoopRegion(ELoopModificationControls Modification);
 
diff --git a/Engine/Plugins/Editor/WaveformEditor/Source/WaveformEditor/Public/WaveformEditorModule.h b/Engine/Plugins/Editor/WaveformEditor/Source/WaveformEditor/Public/WaveformEditorModule.h
--- a/Engine/Plugins/Editor/WaveformEditor/Source/WaveformEditor/Public/WaveformEditorModule.h
+++ b/Engine/Plugins/Editor/WaveformEditor/Source/WaveformEditor/Public/WaveformEditorModule.h
@@ -3,6 +3,7 @@
 #pragma once
 
 #include "IWaveformEditorModule.h"
+#include "Delegates/Delegate.h"
 #include "Templates/SharedPointer.h"
 
 class IWaveformEditorInstantiator;
@@ -17,5 +18,8 @@ private:
 	virtual void RegisterContentBrowserExtensions(IWaveformEditorInstantiator* Instantiator) override;
 
 	TSharedPtr<IWaveformEditorInstantiator> WaveformEditorInstantiator = nullptr;
+
+	FDelegateHandle OnPostEngineInitHandle;
+	FDelegateHandle OnAssetPostImportHandle;
 };
 
diff --git a/Engine/Plugins/Editor/WaveformEditor/Source/WaveformEditor/WaveformEditor.build.cs b/Engine/Plugins/Editor/WaveformEditor/Source/WaveformEditor/WaveformEditor.build.cs
--- a/Engine/Plugins/Editor/WaveformEditor/Source/WaveformEditor/WaveformEditor.build.cs
+++ b/Engine/Plugins/Editor/WaveformEditor/Source/WaveformEditor/WaveformEditor.build.cs
@@ -16,6 +16,7 @@ namespace UnrealBuildTool.Rules
 				{
 					"AssetRegistry",
 					"AudioExtensions",
+					"AudioSynesthesiaCore",
 					"AudioWidgets",
 					"Core",
 					"CoreUObject",
diff --git a/Engine/Plugins/Editor/WaveformEditor/Source/WaveformTransformations/Private/WaveformAudioAnalysisFunctions.cpp b/Engine/Plugins/Editor/WaveformEditor/Source/WaveformTransformations/Private/WaveformAudioAnalysisFunctions.cpp
new file mode 100644
--- /dev/null
+++ b/Engine/Plugins/Editor/WaveformEditor/Source/WaveformTransformations/Private/WaveformAudioAnalysisFunctions.cpp
@@ -0,0 +1,176 @@
+// Copyright Epic Games, Inc. All Rights Reserved.
+
+#include "WaveformAudioAnalysisFunctions.h"
+
+#include "LoudnessFactory.h"
+#include "Algo/MaxElement.h"
+#include "LKFSNRTFactory.h"
+#include "MeterFactory.h"
+#include <limits>
+
+namespace WaveformAudioAnalysis
+{
+	constexpr float DbMin = -std::numeric_limits<float>::infinity();
+
+	// GatedLoudness starts analysis with a -70dB gate, if something falls below the gated value, we should return DbMin
+	constexpr float NoiseFloor = -70.f;
+
+	float GetRMSPeak(const Audio::FAlignedFloatBuffer& InputAudio, float SampleRate, int32 NumChannels)
+	{
+		check(NumChannels > 0);
+		check(SampleRate > 0.f);
+
+		const float AudioLengthSeconds = InputAudio.Num() / NumChannels / SampleRate;
+		const float RMSAnalysisPeriod = 0.3f;
+		const bool bAnalyzeAsSingleBuffer = RMSAnalysisPeriod > AudioLengthSeconds;
+
+		Audio::FMeterSettings Settings;
+		Settings.AnalysisPeriod = bAnalyzeAsSingleBuffer ? AudioLengthSeconds : RMSAnalysisPeriod;
+
+		if (bAnalyzeAsSingleBuffer)
+		{
+			Audio::FMeterAnalyzer Analyzer(SampleRate, NumChannels, Settings);
+			Audio::FMeterAnalyzerResults Results = Analyzer.ProcessAudio(InputAudio);
+			const float* MeterValue = Algo::MaxElement(Results.MeterValues, [](const float A, const float B)
+				{
+					return A < B;
+				});
+
+			if (MeterValue)
+			{
+				if (*MeterValue < NoiseFloor)
+				{
+					return DbMin;
+				}
+
+				return *MeterValue;
+			}
+		}
+		else
+		{
+			Audio::FMeterFactory Analyzer;
+			TUniquePtr<Audio::IAnalyzerResult> Result = Analyzer.NewResult();
+			TUniquePtr<Audio::IAnalyzerWorker> AnalyzerWorker = Analyzer.NewWorker({ static_cast<int32>(SampleRate), NumChannels }, &Settings);
+
+			if (AnalyzerWorker == nullptr)
+			{
+				return DbMin;
+			}
+
+			AnalyzerWorker->Analyze(MakeArrayView(InputAudio), Result.Get());
+
+			Audio::FMeterResult* MeterResult = static_cast<Audio::FMeterResult*>(Result.Get());
+
+			if (MeterResult)
+			{
+				const TArray<Audio::FMeterEntry>& LoudnessArray = MeterResult->GetMeterArray();
+				auto MeterEntry = Algo::MaxElement(LoudnessArray, [](const Audio::FMeterEntry& A, const Audio::FMeterEntry& B)
+					{
+						return A.MeterValue < B.MeterValue;
+					});
+
+				if (MeterEntry)
+				{
+					if (MeterEntry->MeterValue < NoiseFloor)
+					{
+						return DbMin;
+					}
+
+					return MeterEntry->MeterValue;
+				}
+			}
+		}
+
+		return DbMin;
+	}
+
+	float GetLoudnessPeak(const Audio::FAlignedFloatBuffer& InputAudio, float SampleRate, int32 NumChannels)
+	{
+		Audio::FLoudnessFactory Analyzer;
+		Audio::FLoudnessSettings Settings;
+		Settings.AnalysisPeriod = 1.f;
+
+		// Could implement bAnalyzeAsSingleBuffer optimization here as above
+
+		TUniquePtr<Audio::IAnalyzerResult> Result = Analyzer.NewResult();
+		TUniquePtr<Audio::IAnalyzerWorker> AnalyzerWorker = Analyzer.NewWorker({ static_cast<int32>(SampleRate), NumChannels }, &Settings);
+
+		if (AnalyzerWorker == nullptr)
+		{
+			return DbMin;
+		}
+
+		AnalyzerWorker->Analyze(MakeArrayView(InputAudio), Result.Get());
+
+		Audio::FLoudnessResult* LoudnessResult = static_cast<Audio::FLoudnessResult*>(Result.Get());
+
+		if (LoudnessResult)
+		{
+			const TArray<Audio::FLoudnessEntry>& LoudnessArray = LoudnessResult->GetLoudnessArray();
+			auto LoudestEntry = Algo::MaxElement(LoudnessArray, [](const Audio::FLoudnessEntry& A, const Audio::FLoudnessEntry& B)
+				{
+					return A.Loudness < B.Loudness;
+				});
+
+			if (LoudestEntry)
+			{
+				if (LoudestEntry->Loudness < NoiseFloor)
+				{
+					return DbMin;
+				}
+
+				return LoudestEntry->Loudness;
+			}
+		}
+
+		return DbMin;
+	}
+
+	float GetPeakSampleValue(const Audio::FAlignedFloatBuffer& InputAudio)
+	{
+		const float MaxValue = Audio::ArrayMaxAbsValue(InputAudio);
+		float DbValue = Audio::ConvertToDecibels(MaxValue);
+
+		if (DbValue < NoiseFloor)
+		{
+			DbValue = DbMin;
+		}
+
+		return DbValue;
+	}
+
+	float GetLUFS(const Audio::FAlignedFloatBuffer& InputAudio, float SampleRate, int32 NumChannels)
+	{
+		Audio::FLKFSNRTFactory Factory;
+		Audio::FAnalyzerNRTParameters AnalyzerParameters(SampleRate, NumChannels);
+		Audio::FLKFSNRTSettings LKFSSettings;
+
+		TUniquePtr<Audio::IAnalyzerNRTResult> Result = Factory.NewResult();
+		TUniquePtr<Audio::IAnalyzerNRTWorker> Worker = Factory.NewWorker(AnalyzerParameters, &LKFSSettings);
+
+		if (Worker == nullptr)
+		{
+			return DbMin;
+		}
+
+		Worker->Analyze(MakeArrayView(InputAudio), Result.Get());
+		Worker->Finalize(Result.Get());
+
+		Audio::FLKFSNRTResult* LKFSResult = static_cast<Audio::FLKFSNRTResult*>(Result.Get());
+
+		if (LKFSResult)
+		{
+			const float GatedLoudness = LKFSResult->FindAggregateLoudnessStats(LKFSResult->ChannelIndexOverall)->GatedLoudness;
+
+			if (GatedLoudness == FLKFSNRTAggregateStats::InvalidLoudness || GatedLoudness < NoiseFloor)
+			{
+				// Sample not large enough to analyze
+				return DbMin;
+			}
+
+			return GatedLoudness;
+		}
+
+		return DbMin;
+	}
+}
diff --git a/Engine/Plugins/Editor/WaveformEditor/Source/WaveformTransformations/Private/WaveformTransformationNormalize.cpp b/Engine/Plugins/Editor/WaveformEditor/Source/WaveformTransformations/Private/WaveformTransformationNormalize.cpp
--- a/Engine/Plugins/Editor/WaveformEditor/Source/WaveformTransformations/Private/WaveformTransformationNormalize.cpp
+++ b/Engine/Plugins/Editor/WaveformEditor/Source/WaveformTransformations/Private/WaveformTransformationNormalize.cpp
@@ -2,110 +2,12 @@
 
 #include "WaveformTransformationNormalize.h"
 
+#include "DSP/Dsp.h"
 #include "DSP/FloatArrayMath.h"
-#include "AudioAnalyzer.h"
-#include "IAudioAnalyzerInterface.h"
-#include "LoudnessFactory.h"
-#include "Algo/MaxElement.h"
-#include "MeterFactory.h"
+#include "WaveformAudioAnalysisFunctions.h"
 
 #include UE_INLINE_GENERATED_CPP_BY_NAME(WaveformTransformationNormalize)
 
-namespace WaveformTransformNormalizeHelpers
-{
-	float GetRMSPeak(const Audio::FAlignedFloatBuffer& InputAudio, float SampleRate, int32 NumChannels)
-	{
-		const float AudioLengthSeconds = InputAudio.Num() / NumChannels / SampleRate;
-		const float RMSAnalysisPeriod = 0.3f;
-		const bool bAnalyzeAsSingleBuffer = RMSAnalysisPeriod > AudioLengthSeconds;
-
-		Audio::FMeterSettings Settings;
-		Settings.AnalysisPeriod = bAnalyzeAsSingleBuffer ? AudioLengthSeconds : RMSAnalysisPeriod;
-
-		if (bAnalyzeAsSingleBuffer)
-		{
-			Audio::FMeterAnalyzer Analyzer(SampleRate, NumChannels, Settings);
-			Audio::FMeterAnalyzerResults Results = Analyzer.ProcessAudio(InputAudio);
-			const float* MeterValue = Algo::MaxElement(Results.MeterValues, [](const float A, const float B)
-				{
-					return A < B;
-				});
-
-			if (MeterValue)
-			{
-				return *MeterValue;
-			}
-		}
-		else
-		{
-			Audio::FMeterFactory Analyzer;
-			TUniquePtr<Audio::IAnalyzerResult> Result = Analyzer.NewResult();
-			TUniquePtr<Audio::IAnalyzerWorker> AnalyzerWorker = Analyzer.NewWorker({ (int32)SampleRate, NumChannels }, &Settings);
-
-			if (AnalyzerWorker == nullptr)
-			{
-				return 0.f;
-			}
-
-			AnalyzerWorker->Analyze(MakeArrayView(InputAudio), Result.Get());
-
-			Audio::FMeterResult* MeterResult = static_cast<Audio::FMeterResult*>(Result.Get());
-
-			if (MeterResult)
-			{
-				const TArray<Audio::FMeterEntry>& LoudnessArray = MeterResult->GetMeterArray();
-				auto MeterEntry = Algo::MaxElement(LoudnessArray, [](const Audio::FMeterEntry& A, const Audio::FMeterEntry& B)
-					{
-						return A.MeterValue < B.MeterValue;
-					});
-
-				if (MeterEntry)
-				{
-					return MeterEntry->MeterValue;
-				}
-			}
-		}
-		
-		return 0.f;
-	}
-
-	float GetLoudnessPeak(const Audio::FAlignedFloatBuffer& InputAudio, float SampleRate, int32 NumChannels)
-	{
-		Audio::FLoudnessFactory Analyzer;
-		Audio::FLoudnessSettings Settings;
-		Settings.AnalysisPeriod = 1.f;
-
-		TUniquePtr<Audio::IAnalyzerResult> Result = Analyzer.NewResult();
-		TUniquePtr<Audio::IAnalyzerWorker> AnalyzerWorker = Analyzer.NewWorker({ (int32)SampleRate, NumChannels }, &Settings);
-
-		if (AnalyzerWorker == nullptr)
-		{
-			return 0.f;
-		}
-
-		AnalyzerWorker->Analyze(MakeArrayView(InputAudio), Result.Get());
-
-		Audio::FLoudnessResult* LoudnessResult = static_cast<Audio::FLoudnessResult*>(Result.Get());
-
-		if (LoudnessResult)
-		{
-			const TArray<Audio::FLoudnessEntry>& LoudnessArray = LoudnessResult->GetLoudnessArray();
-			auto LoudestEntry = Algo::MaxElement(LoudnessArray, [](const Audio::FLoudnessEntry& A, const Audio::FLoudnessEntry& B)
-				{
-					return A.Loudness < B.Loudness;
-				});
-
-			if (LoudestEntry)
-			{
-				return LoudestEntry->Loudness;
-			}
-		}
-
-		return 0.f;
-	}
-}
-
-
 FWaveTransformationNormalize::FWaveTransformationNormalize(float InTarget, float InMaxGain, ENormalizationMode InMode)
 	: Target(InTarget)
 	, MaxGain(InMaxGain)
@@ -129,11 +31,11 @@ FWaveTransformationNormalize::FWaveTransformationNormalize(float InTarget, float
 			break;
 		}
 	case ENormalizationMode::RMS:
-		PeakDecibelValue = WaveformTransformNormalizeHelpers::GetRMSPeak(InputAudio, InOutWaveInfo.SampleRate, InOutWaveInfo.NumChannels);
+		PeakDecibelValue = WaveformAudioAnalysis::GetRMSPeak(InputAudio, InOutWaveInfo.SampleRate, InOutWaveInfo.NumChannels);
 		PeakDecibelValue = Audio::ConvertToDecibels(PeakDecibelValue);
 		break;
 	case ENormalizationMode::DWeightedLoudness:
-		PeakDecibelValue = WaveformTransformNormalizeHelpers::GetLoudnessPeak(InputAudio,  InOutWaveInfo.SampleRate, InOutWaveInfo.NumChannels);
+		PeakDecibelValue = WaveformAudioAnalysis::GetLoudnessPeak(InputAudio, InOutWaveInfo.SampleRate, InOutWaveInfo.NumChannels);
 		break;
 	case ENormalizationMode::COUNT:
 	default:
diff --git a/Engine/Plugins/Editor/WaveformEditor/Source/WaveformTransformations/Public/WaveformAudioAnalysisFunctions.h b/Engine/Plugins/Editor/WaveformEditor/Source/WaveformTransformations/Public/WaveformAudioAnalysisFunctions.h
new file mode 100644
--- /dev/null
+++ b/Engine/Plugins/Editor/WaveformEditor/Source/WaveformTransformations/Public/WaveformAudioAnalysisFunctions.h
@@ -0,0 +1,21 @@
+// Copyright Epic Games, Inc. All Rights Reserved.
+
+#pragma once
+
+#include "DSP/AlignedBuffer.h"
+
+#define UE_API WAVEFORMTRANSFORMATIONS_API
+
+namespace WaveformAudioAnalysis
+{
+	UE_API float GetRMSPeak(const Audio::FAlignedFloatBuffer& InputAudio, float SampleRate, int32 NumChannels);
+
+	UE_API float GetLoudnessPeak(const Audio::FAlignedFloatBuffer& InputAudio, float SampleRate, int32 NumChannels);
+
+	UE_API float GetPeakSampleValue(const Audio::FAlignedFloatBuffer& InputAudio);
+
+	// Gated Loudness K-Weighted relative to Full Scale (Not-Real-Time)
+	UE_API float GetLUFS(const Audio::FAlignedFloatBuffer& InputAudio, float SampleRate, int32 NumChannels);
+}
+
+#undef UE_API
\ No newline at end of file
diff --git a/Engine/Source/Runtime/Engine/Classes/Sound/SoundWave.h b/Engine/Source/Runtime/Engine/Classes/Sound/SoundWave.h
--- a/Engine/Source/Runtime/Engine/Classes/Sound/SoundWave.h
+++ b/Engine/Source/Runtime/Engine/Classes/Sound/SoundWave.h
@@ -496,6 +496,9 @@ public:
 
 	ENGINE_API const ESoundWaveCuePointOrigin GetCuePointOrigin() const;
 	ENGINE_API void SetCuePointOrigin(const ESoundWaveCuePointOrigin InCuePointOrigin);
+
+	/** Set gated LUFS and SamplePeakDB values */
+	ENGINE_API void SetLoudnessValues(const float InLUFS, const float InSamplePeakDB);
 #endif //WITH_EDITOR
 
 	/** Filters for the cue points that _are_ loop regions and returns those as a new array */
@@ -779,6 +782,14 @@ public:
 	UPROPERTY()
 	TArray<int32> ChannelSizes;
 
+	/** Gated Loudness Units Full Scale of the SoundWave. Can be modified by transformations. */
+	UPROPERTY(Category = Info, AssetRegistrySearchable, VisibleAnywhere)
+	float LUFS;
+
+	/** Highest Sample Peak in db of the SoundWave. Can be modified by transformations. */
+	UPROPERTY(Category = Info, AssetRegistrySearchable, VisibleAnywhere)
+	float SamplePeakDB;
+
 #endif // WITH_EDITORONLY_DATA
 
 protected:
diff --git a/Engine/Source/Runtime/Engine/Private/SoundWave.cpp b/Engine/Source/Runtime/Engine/Private/SoundWave.cpp
--- a/Engine/Source/Runtime/Engine/Private/SoundWave.cpp
+++ b/Engine/Source/Runtime/Engine/Private/SoundWave.cpp
@@ -1560,6 +1560,12 @@ void USoundWave::SetCuePointOrigin(const ESoundWaveCuePointOrigin InCuePointOrig
 {
 	CuePointOrigin = InCuePointOrigin;
 }
+
+void USoundWave::SetLoudnessValues(const float InLUFS, const float InSamplePeakDB)
+{
+	LUFS = InLUFS;
+	SamplePeakDB = InSamplePeakDB;
+}
 #endif //WITH_EDITOR
 
 TArray<FSoundWaveCuePoint> USoundWave::GetLoopRegions() const
@@ -1775,6 +1781,8 @@ bool USoundWave::GetImportedSoundWaveData(TArray<uint8>& OutRawPCMData, uint32&
 					NumFrames = WaveInfo.SampleDataSize / sizeof(int16);
 					NumSamples = NumFrames * NumChannels;
 
+					check(NumSamples % NumChannels == 0);
+
 					if (NumSamples > 0)
 					{
 						// Translate NumSamples to bytes because OutRawPCMData is in bytes. 
@@ -1822,6 +1830,9 @@ bool USoundWave::GetImportedSoundWaveData(TArray<uint8>& OutRawPCMData, uint32&
 				return false;
 			}
 
+			check(NumChannels > 0);
+			check(WaveInfo.SampleDataSize% NumChannels == 0);
+
 			// Copy the raw PCM data and the header info that was parsed
 			OutRawPCMData.AddUninitialized(WaveInfo.SampleDataSize);
 			FMemory::Memcpy(OutRawPCMData.GetData(), WaveInfo.SampleDataStart, WaveInfo.SampleDataSize);
