FadeCurves being overwritten in PostLoad fix

diff --git a/Engine/Plugins/Editor/WaveformEditor/Source/WaveformTransformations/Private/WaveformTransformationTrimFade.cpp b/Engine/Plugins/Editor/WaveformEditor/Source/WaveformTransformations/Private/WaveformTransformationTrimFade.cpp
--- a/Engine/Plugins/Editor/WaveformEditor/Source/WaveformTransformations/Private/WaveformTransformationTrimFade.cpp
+++ b/Engine/Plugins/Editor/WaveformEditor/Source/WaveformTransformations/Private/WaveformTransformationTrimFade.cpp
@@ -333,6 +333,7 @@ void UWaveformTransformationTrimFade::PostLoad()
 	const float Linear = 1.f;
 
 	PRAGMA_DISABLE_DEPRECATION_WARNINGS
+	// Overwrite exiting fade in function data with deprecated curve properties to migrate the data to the new properties
 	if (StartFadeTime != 0.f)
 	{
 		if (StartFadeCurve > Linear)
@@ -358,13 +359,14 @@ void UWaveformTransformationTrimFade::PostLoad()
 		StartFadeTime = 0.f;
 	}
 	PRAGMA_ENABLE_DEPRECATION_WARNINGS
-	else
+	else if (FadeFunctions.FadeIn == nullptr)
 	{
 		FadeFunctions.FadeIn = NewObject<UFadeCurveFunctionExponential>(this, UFadeCurveFunctionExponential::StaticClass(), NAME_None, RF_Transactional);
 		FadeFunctions.FadeIn->Duration = 0.f;
 	}
 
 	PRAGMA_DISABLE_DEPRECATION_WARNINGS
+	// Overwrite exiting fade out function data with deprecated curve properties to migrate the data to the new properties
 	if (EndFadeTime != 0.f)
 	{
 		if (EndFadeCurve > Linear)
@@ -390,7 +392,7 @@ void UWaveformTransformationTrimFade::PostLoad()
 		EndFadeTime = 0.f;
 	}
 	PRAGMA_ENABLE_DEPRECATION_WARNINGS
-	else
+	else if (FadeFunctions.FadeOut == nullptr)
 	{
 		FadeFunctions.FadeOut = NewObject<UFadeCurveFunctionExponential>(this, UFadeCurveFunctionExponential::StaticClass(), NAME_None, RF_Transactional);
 		FadeFunctions.FadeOut->Duration = 0.f;
