[Waveform Editor] - Normalizing values in a waveform does not update the Sound Waves MetaData

diff --git a/Engine/Plugins/Editor/WaveformEditor/Source/WaveformEditor/Private/WaveformEditor.cpp b/Engine/Plugins/Editor/WaveformEditor/Source/WaveformEditor/Private/WaveformEditor.cpp
--- a/Engine/Plugins/Editor/WaveformEditor/Source/WaveformEditor/Private/WaveformEditor.cpp
+++ b/Engine/Plugins/Editor/WaveformEditor/Source/WaveformEditor/Private/WaveformEditor.cpp
@@ -137,28 +137,6 @@ bool FWaveformEditor::Init(const EToolkitMode::Type Mode, const TSharedPtr<ITool
 
 				bCueChangeRegisteredByWaveformEditor = true;
 			}
-
-			if (!MarkerTransformation->OnTransformationChanged.IsBound())
-			{
-				MarkerTransformation->OnTransformationChanged.BindLambda([WeakWaveformEditor](bool bModify)
-					{
-						if (TSharedPtr<FWaveformEditor> PinnedEditor = WeakWaveformEditor.Pin())
-						{
-							PinnedEditor->UpdateTransformations(bModify);
-
-						}
-					});
-			}
-			if (!MarkerTransformation->OnTransformationRenderChanged.IsBound())
-			{
-				MarkerTransformation->OnTransformationRenderChanged.BindLambda([WeakWaveformEditor]()
-					{
-						if (TSharedPtr<FWaveformEditor> PinnedEditor = WeakWaveformEditor.Pin())
-						{
-							PinnedEditor->UpdateRenderElements();
-						}
-					});
-			}
 		}
 		else
 		{
@@ -166,30 +144,8 @@ bool FWaveformEditor::Init(const EToolkitMode::Type Mode, const TSharedPtr<ITool
 			bMarkerTransformationIsActive = false;
 		}
 
-		if (TrimFadeTransformation != nullptr)
-		{
-			if (!TrimFadeTransformation->OnTransformationChanged.IsBound())
-			{
-				TrimFadeTransformation->OnTransformationChanged.BindLambda([WeakWaveformEditor](bool bModify)
-					{
-						if (TSharedPtr<FWaveformEditor> PinnedEditor = WeakWaveformEditor.Pin())
-						{
-							PinnedEditor->UpdateTransformations(bModify);
-						}
-					});
-			}
+		BindDelegatesForActiveTransformations();
 
-			if (!TrimFadeTransformation->OnTransformationRenderChanged.IsBound())
-			{
-				TrimFadeTransformation->OnTransformationRenderChanged.BindLambda([WeakWaveformEditor]()
-					{
-						if (TSharedPtr<FWaveformEditor> PinnedEditor = WeakWaveformEditor.Pin())
-						{
-							PinnedEditor->UpdateRenderElements();
-						}
-					});
-			}
-		}
 #if WITH_EDITOR
 		UpdateSoundWaveProperties();
 #endif // WITH_EDITOR
@@ -211,15 +167,6 @@ FWaveformEditor::~FWaveformEditor()
 
 	if (MarkerTransformation)
 	{
-		if (MarkerTransformation->OnTransformationChanged.IsBound())
-		{
-			MarkerTransformation->OnTransformationChanged.Unbind();
-		}
-		if (MarkerTransformation->OnTransformationRenderChanged.IsBound())
-		{
-			MarkerTransformation->OnTransformationRenderChanged.Unbind();
-		}
-
 		if (bCueChangeRegisteredByWaveformEditor)
 		{
 			check(MarkerTransformation->Markers);
@@ -244,15 +191,6 @@ FWaveformEditor::~FWaveformEditor()
 
 	if (TrimFadeTransformation)
 	{
-		if (TrimFadeTransformation->OnTransformationChanged.IsBound())
-		{
-			TrimFadeTransformation->OnTransformationChanged.Unbind();
-		}
-		if (TrimFadeTransformation->OnTransformationRenderChanged.IsBound())
-		{
-			TrimFadeTransformation->OnTransformationRenderChanged.Unbind();
-		}
-
 		// If the user has not used the trimfade then remove the transformation to avoid unnecessary transformations in the chain
 		check(TrimFadeTransformation != nullptr);
 		check(SoundWave != nullptr);
@@ -275,6 +213,8 @@ FWaveformEditor::~FWaveformEditor()
 		}
 	}
 
+	UnBindDelegatesForActiveTransformations();
+
 	if (FReimportManager::Instance())
 	{
 		FReimportManager::Instance()->OnPostReimport().RemoveAll(this);
@@ -461,6 +401,59 @@ TObjectPtr<UWaveformTransformationMarkers> FWaveformEditor::AddMarkerTransformat
 	return MarkersTransformation;
 }
 
+void FWaveformEditor::BindDelegatesForActiveTransformations()
+{
+	check(SoundWave);
+	TWeakPtr<FWaveformEditor> WeakWaveformEditor = SharedThis(this);
+
+	for (const TObjectPtr<UWaveformTransformationBase>& Transformation : SoundWave->Transformations)
+	{
+		if (Transformation != nullptr)
+		{
+			if (!Transformation->OnTransformationChanged.IsBound())
+			{
+				Transformation->OnTransformationChanged.BindLambda([WeakWaveformEditor](bool bModify)
+					{
+						if (TSharedPtr<FWaveformEditor> PinnedEditor = WeakWaveformEditor.Pin())
+						{
+							PinnedEditor->UpdateTransformations(bModify);
+						}
+					});
+			}
+
+			if (!Transformation->OnTransformationRenderChanged.IsBound())
+			{
+				Transformation->OnTransformationRenderChanged.BindLambda([WeakWaveformEditor]()
+					{
+						if (TSharedPtr<FWaveformEditor> PinnedEditor = WeakWaveformEditor.Pin())
+						{
+							PinnedEditor->UpdateRenderElements();
+						}
+					});
+			}
+		}
+	}
+}
+
+void FWaveformEditor::UnBindDelegatesForActiveTransformations()
+{
+	check(SoundWave);
+	for (const TObjectPtr<UWaveformTransformationBase>& Transformation : SoundWave->Transformations)
+	{
+		if (Transformation != nullptr)
+		{
+			if (Transformation->OnTransformationChanged.IsBound())
+			{
+				Transformation->OnTransformationChanged.Unbind();
+			}
+			if (Transformation->OnTransformationRenderChanged.IsBound())
+			{
+				Transformation->OnTransformationRenderChanged.Unbind();
+			}
+		}
+	}
+}
+
 void FWaveformEditor::ToggleFadeIn()
 {
 	check(SoundWave);
@@ -1712,27 +1705,6 @@ void FWaveformEditor::NotifyPostChange(const FPropertyChangedEvent& PropertyChan
 
 			bCueChangeRegisteredByWaveformEditor = true;
 		}
-
-		if (!MarkerTransformation->OnTransformationChanged.IsBound())
-		{
-			MarkerTransformation->OnTransformationChanged.BindLambda([WeakWaveformEditor](bool bModify)
-				{
-					if (TSharedPtr<FWaveformEditor> PinnedEditor = WeakWaveformEditor.Pin())
-					{
-						PinnedEditor->UpdateTransformations(bModify);
-					}
-				});
-		}
-		if (!MarkerTransformation->OnTransformationRenderChanged.IsBound())
-		{
-			MarkerTransformation->OnTransformationRenderChanged.BindLambda([WeakWaveformEditor]()
-				{
-					if (TSharedPtr<FWaveformEditor> PinnedEditor = WeakWaveformEditor.Pin())
-					{
-						PinnedEditor->UpdateRenderElements();
-					}
-				});
-		}
 	}
 	else
 	{
@@ -1740,32 +1712,7 @@ void FWaveformEditor::NotifyPostChange(const FPropertyChangedEvent& PropertyChan
 		bMarkerTransformationIsActive = false;
 	}
 
-	TObjectPtr<UWaveformTransformationTrimFade> TrimFadeTransformation = GetTrimFadeTransformation();
-
-	if (TrimFadeTransformation != nullptr)
-	{
-		if (!TrimFadeTransformation->OnTransformationChanged.IsBound())
-		{
-			TrimFadeTransformation->OnTransformationChanged.BindLambda([WeakWaveformEditor](bool bModify)
-				{
-					if (TSharedPtr<FWaveformEditor> PinnedEditor = WeakWaveformEditor.Pin())
-					{
-						PinnedEditor->UpdateTransformations(bModify);
-					}
-				});
-		}
-
-		if (!TrimFadeTransformation->OnTransformationRenderChanged.IsBound())
-		{
-			TrimFadeTransformation->OnTransformationRenderChanged.BindLambda([WeakWaveformEditor]()
-				{
-					if (TSharedPtr<FWaveformEditor> PinnedEditor = WeakWaveformEditor.Pin())
-					{
-						PinnedEditor->UpdateRenderElements();
-					}
-				});
-		}
-	}
+	BindDelegatesForActiveTransformations();
 
 	if ((PropertyChangedEvent.ChangeType & EPropertyChangeType::ValueSet) == EPropertyChangeType::ValueSet
 	 && (PropertyChangedEvent.ChangeType & EPropertyChangeType::Interactive) != EPropertyChangeType::Interactive)
diff --git a/Engine/Plugins/Editor/WaveformEditor/Source/WaveformEditor/Public/WaveformEditor.h b/Engine/Plugins/Editor/WaveformEditor/Source/WaveformEditor/Public/WaveformEditor.h
--- a/Engine/Plugins/Editor/WaveformEditor/Source/WaveformEditor/Public/WaveformEditor.h
+++ b/Engine/Plugins/Editor/WaveformEditor/Source/WaveformEditor/Public/WaveformEditor.h
@@ -118,6 +118,10 @@ private:
 	UE_API TObjectPtr<UWaveformTransformationMarkers> GetMarkerTransformation() const;
 	UE_API TObjectPtr<UWaveformTransformationMarkers> AddMarkerTransformation();
 
+	// Get all active transformations and bind the waveform editor delegates
+	UE_API void BindDelegatesForActiveTransformations();
+	UE_API void UnBindDelegatesForActiveTransformations();
+
 	UE_API void ToggleFadeIn();
 	UE_API bool CanFadeIn();
 
diff --git a/Engine/Source/Runtime/AudioExtensions/Private/IWaveformTransformation.cpp b/Engine/Source/Runtime/AudioExtensions/Private/IWaveformTransformation.cpp
--- a/Engine/Source/Runtime/AudioExtensions/Private/IWaveformTransformation.cpp
+++ b/Engine/Source/Runtime/AudioExtensions/Private/IWaveformTransformation.cpp
@@ -26,7 +26,15 @@ void UWaveformTransformationBase::PostEditUndo()
 {
 	Super::PostEditUndo();
 
-	static const bool bMarkFileDirty = true;
+	constexpr bool bMarkFileDirty = true;
+	OnTransformationChanged.ExecuteIfBound(bMarkFileDirty);
+}
+
+void UWaveformTransformationBase::PostEditChangeProperty(FPropertyChangedEvent& PropertyChangedEvent)
+{
+	Super::PostEditChangeProperty(PropertyChangedEvent);
+
+	constexpr bool bMarkFileDirty = true;
 	OnTransformationChanged.ExecuteIfBound(bMarkFileDirty);
 }
 
diff --git a/Engine/Source/Runtime/AudioExtensions/Public/IWaveformTransformation.h b/Engine/Source/Runtime/AudioExtensions/Public/IWaveformTransformation.h
--- a/Engine/Source/Runtime/AudioExtensions/Public/IWaveformTransformation.h
+++ b/Engine/Source/Runtime/AudioExtensions/Public/IWaveformTransformation.h
@@ -156,6 +156,7 @@ public:
 	virtual void GetTransformationInfo(FWaveformTransformationInfo& InOutTransformationInfo) const {};
 
 	AUDIOEXTENSIONS_API virtual void PostEditUndo() override;
+	AUDIOEXTENSIONS_API virtual void PostEditChangeProperty(FPropertyChangedEvent& PropertyChangedEvent) override;
 
 	void NotifyPropertyChange(FProperty* Property);
 #endif //WITH_EDITOR
