[Waveform Editor] - Adding a Default Transformation does not auto sort based on hierarchy

diff --git a/Engine/Plugins/Editor/WaveformEditor/Source/WaveformEditor/Private/WaveformEditor.cpp b/Engine/Plugins/Editor/WaveformEditor/Source/WaveformEditor/Private/WaveformEditor.cpp
--- a/Engine/Plugins/Editor/WaveformEditor/Source/WaveformEditor/Private/WaveformEditor.cpp
+++ b/Engine/Plugins/Editor/WaveformEditor/Source/WaveformEditor/Private/WaveformEditor.cpp
@@ -319,6 +319,8 @@ TObjectPtr<UWaveformTransformationTrimFade> FWaveformEditor::GetOrAddTrimFadeTra
 
 		TrimFadeTransformation->FadeFunctions.FadeIn->Duration = 0.f;
 		TrimFadeTransformation->FadeFunctions.FadeOut->Duration = 0.f;
+
+		TransformationChainConfig = SoundWave->UpdateTransformations();
 	}
 
 	return TrimFadeTransformation;
