Merge from main:
UE-306853, UE-318566
Scrubber does not move inside Isolated Region when converting a Marker to Loop Region
More generally, keep selection in sync with the most recently edited property, viewport<->details panel.
Address Bughawk feedback, WaveformTransformationMarkers.cpp

diff --git a/Engine/Plugins/Editor/WaveformEditor/Source/WaveformTransformations/Private/WaveformTransformationMarkers.cpp b/Engine/Plugins/Editor/WaveformEditor/Source/WaveformTransformations/Private/WaveformTransformationMarkers.cpp
--- a/Engine/Plugins/Editor/WaveformEditor/Source/WaveformTransformations/Private/WaveformTransformationMarkers.cpp
+++ b/Engine/Plugins/Editor/WaveformEditor/Source/WaveformTransformations/Private/WaveformTransformationMarkers.cpp
@@ -494,6 +494,7 @@ void UWaveformTransformationMarkers::SetLoopPreviewing()
 	check(Markers);
 	bool bIsSelectedCuePresent = false;
 	bool bTransformationChanged = false;
+	bool bTransformationChangedWithoutMarkingDirty = false;
 
 	for (FSoundWaveCuePoint& CuePoint : Markers->CuesAndLoops)
 	{
@@ -515,6 +516,12 @@ void UWaveformTransformationMarkers::SetLoopPreviewing()
 				bTransformationChanged = true;
 			}
 
+			if (Markers->SelectedCue == CuePoint.CuePointID && !IsLoopRegionActive())
+			{
+				bTransformationChangedWithoutMarkingDirty = true;
+				PreviewSelectedLoop();
+			}
+
 			// Update LoopRegion Preview window if a loop is selected and the user edits the loop bounds using the properties window
 			if (IsLoopRegionActive() && Markers->SelectedCue == CuePoint.CuePointID)
 			{
@@ -554,6 +561,10 @@ void UWaveformTransformationMarkers::SetLoopPreviewing()
 	{
 		OnTransformationChanged.ExecuteIfBound(true);
 	}
+	else if (bTransformationChangedWithoutMarkingDirty)
+	{
+		OnTransformationChanged.ExecuteIfBound(false);
+	}
 }
 
 bool UWaveformTransformationMarkers::ResetLoopPreviewing()
@@ -736,6 +747,7 @@ void UWaveformTransformationMarkers::SelectLoopRegionByKeyboard(const FKey& Pres
 	}
 }
 
+#if WITH_EDITOR
 void UWaveCueArray::PostEditChangeProperty(FPropertyChangedEvent& PropertyChangedEvent)
 {
 	Super::PostEditChangeProperty(PropertyChangedEvent);
@@ -787,6 +799,32 @@ void UWaveCueArray::PostEditChangeProperty(FPropertyChangedEvent& PropertyChange
 	CueChanged.ExecuteIfBound();
 }
 
+void UWaveCueArray::PostEditChangeChainProperty(struct FPropertyChangedChainEvent& PropertyChangedEvent)
+{
+	Super::PostEditChangeChainProperty(PropertyChangedEvent);
+
+	// Clear selection if property change type is ArrayRemove, ArrayClear, ToggleEditable (edit condition state has changed), or Unspecified.
+	SelectedCue = INDEX_NONE;
+	if ( PropertyChangedEvent.ChangeType & (EPropertyChangeType::ArrayAdd
+										   |EPropertyChangeType::ResetToDefault
+										   |EPropertyChangeType::ValueSet
+										   |EPropertyChangeType::Duplicate
+										   |EPropertyChangeType::Interactive
+										   |EPropertyChangeType::ArrayMove
+										   |EPropertyChangeType::Redirected) )
+	{
+		if (CuesAndLoops.Num() > 0 && PropertyChangedEvent.PropertyChain.GetActiveMemberNode() != nullptr)
+		{
+			int32 MemberNodeIndex = PropertyChangedEvent.GetArrayIndex(PropertyChangedEvent.PropertyChain.GetActiveMemberNode()->GetValue()->GetName());
+			if (MemberNodeIndex >= 0 && MemberNodeIndex < CuesAndLoops.Num())
+			{
+				SelectedCue = CuesAndLoops[MemberNodeIndex].CuePointID;
+			}
+		}
+	}
+}
+#endif // WITH_EDITOR
+
 void UWaveCueArray::InitMarkersIfNotSet(const TArray<FSoundWaveCuePoint>& InMarkers)
 {
 	//Prevent USoundWave from overwriting the Transformation unintentionally
diff --git a/Engine/Plugins/Editor/WaveformEditor/Source/WaveformTransformations/Public/WaveformTransformationMarkers.h b/Engine/Plugins/Editor/WaveformEditor/Source/WaveformTransformations/Public/WaveformTransformationMarkers.h
--- a/Engine/Plugins/Editor/WaveformEditor/Source/WaveformTransformations/Public/WaveformTransformationMarkers.h
+++ b/Engine/Plugins/Editor/WaveformEditor/Source/WaveformTransformations/Public/WaveformTransformationMarkers.h
@@ -34,8 +34,12 @@ class UWaveCueArray : public UObject
 {
 	GENERATED_BODY()
 public:
-	UE_API void PostEditChangeProperty(FPropertyChangedEvent& PropertyChangedEvent);
-	
+#if WITH_EDITOR
+	UE_API void PostEditChangeProperty(FPropertyChangedEvent& PropertyChangedEvent) override;
+
+	UE_API void PostEditChangeChainProperty(struct FPropertyChangedChainEvent& PropertyChangedEvent) override;
+#endif // WITH_EDITOR
+
 	//If uninitialized, init Marker array
 	UE_API void InitMarkersIfNotSet(const TArray<FSoundWaveCuePoint>& InMarkers);
 	//Uninitialize and empty Marker array
@@ -107,9 +111,10 @@ public:
 	// When EndLoopTime is < 0 it skips processing the audio in FWaveTransformationMarkers::ProcessAudio 
 	UPROPERTY(EditAnywhere, Category = "Loop Preview")
 	double EndLoopTime = -1.0;
+protected:
 	UPROPERTY(EditAnywhere, Category = "Loop Preview")
 	bool bIsPreviewingLoopRegion = false;
-
+public:
 	UE_API virtual Audio::FTransformationPtr CreateTransformation() const override;
 
 	UE_API void UpdateConfiguration(FWaveTransformUObjectConfiguration& InOutConfiguration) override;
@@ -144,7 +149,7 @@ public:
 	UE_API FSoundWaveCuePoint* GetSelectedMarker() const;
 
 private:
-
+	
 	UE_API void PreviewSelectedLoop();
 	UE_API void SetIsPreviewingLoopRegion(double InStartTime, double InEndTime, bool bIsPreviewing);
 	UE_API void SelectLoopRegionByKeyboard(const FKey& PressedKey);
diff --git a/Engine/Plugins/Editor/WaveformEditor/Source/WaveformTransformationsWidgets/Private/WaveformTransformationMarkersRenderer.cpp b/Engine/Plugins/Editor/WaveformEditor/Source/WaveformTransformationsWidgets/Private/WaveformTransformationMarkersRenderer.cpp
--- a/Engine/Plugins/Editor/WaveformEditor/Source/WaveformTransformationsWidgets/Private/WaveformTransformationMarkersRenderer.cpp
+++ b/Engine/Plugins/Editor/WaveformEditor/Source/WaveformTransformationsWidgets/Private/WaveformTransformationMarkersRenderer.cpp
@@ -456,7 +456,7 @@ FCursorReply FWaveformTransformationMarkerRenderer::OnCursorQuery(const FGeometr
 		return FCursorReply::Unhandled();
 	}
 
-	// Set resize on loop handle hover, loops are seperated as we want highlight interactions to superscede region highlights
+	// Set resize on loop handle hover, loops are separated as we want highlight interactions to superscede region highlights
 	for (const FSoundWaveCuePoint& CuePoint : MarkersArray->CuesAndLoops)
 	{
 		const bool bIsLeftHandleHighlighted = IsPositionInLoopHandleArea(MousePosition, MyGeometry, CuePoint, true) || 
@@ -712,10 +712,7 @@ void FWaveformTransformationMarkerRenderer::SetIsPreviewingLoopRegion(double InS
 	check(StrongMarkersTransformation);
 	// Stop any current loop previewing so StartLoopTime is never greater than EndLoopTime (Other than when EndLoopTime is invalid)
 	StrongMarkersTransformation->ResetLoopPreviewing();
-
-	StrongMarkersTransformation->bIsPreviewingLoopRegion = bIsPreviewing;
-	StrongMarkersTransformation->StartLoopTime = InStartTime;
-	StrongMarkersTransformation->EndLoopTime = InEndTime;
+	StrongMarkersTransformation->SetLoopPreviewing();
 }
 
 FVector2D FWaveformTransformationMarkerRenderer::GetLocalCursorPosition(const FPointerEvent& MouseEvent, const FGeometry& EventGeometry)
@@ -868,10 +865,17 @@ void FWaveformTransformationMarkerRenderer::ShowConvertToLoopAtMarker(const FGeo
 		FSlateApplication::Get().OnApplicationActivationStateChanged().Remove(ApplicationActivationStateHandle);
 	}
 
+	const float ContextMenuMaxSize = 750.0f;
+	const FVector2D ContexMenuMaxPosition(LocalWindowMaxPosition.X > ContextMenuMaxSize ? LocalWindowMaxPosition.X - ContextMenuMaxSize : 0.0f
+	                                     ,LocalWindowMaxPosition.Y > ContextMenuMaxSize ? LocalWindowMaxPosition.Y - ContextMenuMaxSize : 0.0f);
+	const FVector2D ContexMenuPosition( FMath::Min(LocalCursorPosition.X, ContexMenuMaxPosition.X),
+										FMath::Min(LocalCursorPosition.Y, ContexMenuMaxPosition.Y) );
+
 	TSharedRef<SWindow> MenuWindow = SNew(SWindow)
 		.AutoCenter(EAutoCenter::None)
-		.ScreenPosition(LocalCursorPosition)
-		.SizingRule(ESizingRule::FixedSize)
+		.Type(EWindowType::Menu)
+		.ScreenPosition(ContexMenuPosition)
+		.SizingRule(ESizingRule::Autosized)
 		.SupportsTransparency(EWindowTransparency::None)
 		.IsPopupWindow(true)
 		.CreateTitleBar(false)
@@ -886,7 +890,7 @@ void FWaveformTransformationMarkerRenderer::ShowConvertToLoopAtMarker(const FGeo
 		.AutoHeight()
 		[
 			SNew(SButton)
-				.Text(LOCTEXT("WaveformTransformationMarkerRendererToggleLoop", "Toggle Loop Region"))
+				.Text(LOCTEXT("WaveformTransformationMarkerRendererToggleLoop", "Toggle between Marker - Loop Region"))
 				.OnClicked_Lambda([this, InCreateLoopMenuWindow = CreateLoopMenuWindow, WeakMarkersTransformation, CueToModify, MinLoopSize = MarkerInInteractionRange.GetMinLoopSize()]() -> FReply
 					{
 						if (TStrongObjectPtr<UWaveformTransformationMarkers> LockedMarkers = WeakMarkersTransformation.Pin())
@@ -905,8 +909,13 @@ void FWaveformTransformationMarkerRenderer::ShowConvertToLoopAtMarker(const FGeo
 
 							LockedMarkers->Markers->EnableLoopRegion(CueToModify, !CueToModify->IsLoopRegion());
 
-							SetActiveLoopRegion(LockedMarkers->Markers->SelectedCue);
-							StrongMarkersTransformation->OnTransformationRenderChanged.ExecuteIfBound();
+							LockedMarkers->Markers->SelectedCue = CueToModify->CuePointID;
+
+							if (StrongMarkersTransformation != nullptr)
+							{
+								SetActiveLoopRegion(LockedMarkers->Markers->SelectedCue);
+								StrongMarkersTransformation->OnTransformationRenderChanged.ExecuteIfBound();
+							}
 						}
 
 						if (GEditor && GEditor->Trans)
@@ -920,7 +929,6 @@ void FWaveformTransformationMarkerRenderer::ShowConvertToLoopAtMarker(const FGeo
 						}
 						return FReply::Handled();
 					})
-
 		];
 
 	MenuContent->AddSlot()
