[Waveform Editor] - Able to add a default Trim Fade in Editor Preferences while Trim Fade is already a default

diff --git a/Engine/Plugins/Editor/WaveformEditor/Source/WaveformEditor/Private/WaveformEditor.cpp b/Engine/Plugins/Editor/WaveformEditor/Source/WaveformEditor/Private/WaveformEditor.cpp
--- a/Engine/Plugins/Editor/WaveformEditor/Source/WaveformEditor/Private/WaveformEditor.cpp
+++ b/Engine/Plugins/Editor/WaveformEditor/Source/WaveformEditor/Private/WaveformEditor.cpp
@@ -239,6 +239,13 @@ void FWaveformEditor::AddDefaultTransformations()
 			
 			if (TransformationClass)
 			{
+				if (TransformationClass == UWaveformTransformationTrimFade::StaticClass())
+				{
+					UE_LOG(LogWaveformEditor, Warning, TEXT("UWaveformTransformationTrimFade is already added automatically and removed if unused. Remove it from LaunchTransformations TSet."));
+
+					continue;
+				}
+
 				EObjectFlags MaskedOuterFlags = SoundWave ? SoundWave->GetMaskedFlags(RF_PropagateToSubObjects) : RF_NoFlags;
 
 				if (SoundWave->HasAnyFlags(RF_ClassDefaultObject | RF_ArchetypeObject))
diff --git a/Engine/Plugins/Editor/WaveformEditor/Source/WaveformEditor/Private/WaveformEditorTransformationsSettings.h b/Engine/Plugins/Editor/WaveformEditor/Source/WaveformEditor/Private/WaveformEditorTransformationsSettings.h
--- a/Engine/Plugins/Editor/WaveformEditor/Source/WaveformEditor/Private/WaveformEditorTransformationsSettings.h
+++ b/Engine/Plugins/Editor/WaveformEditor/Source/WaveformEditor/Private/WaveformEditorTransformationsSettings.h
@@ -31,8 +31,9 @@ public:
 	 * A Transformation chain that will be added to the inspected Soundwave if there aren't any  
 	 * Using TSet to limit each transformation type to one instance to reduce complexity and create a more consistent experience
 	 * USoundWave::Transformations is a TArray because changing it to TSet will delete user data
+	 * UWaveformTransformationTrimFade is disallowed because it is already added automatically and removed if unused
 	 */
-	UPROPERTY(config, EditAnywhere, Category = "Launch Options")
+	UPROPERTY(config, EditAnywhere, Category = "Launch Options", meta = (DisallowedClasses = "/Script/WaveformTransformations.WaveformTransformationTrimFade"))
 	TSet<TSubclassOf<UWaveformTransformationBase>> LaunchTransformations;
 
 	UPROPERTY(config, EditAnywhere, Category = "Zero Crossing Options", meta = (ClampMin = "-144.0", ClampMax = "0.0"))
