[Waveform Editor] - Metadata does not update after reloading a Sound Wave

diff --git a/Engine/Plugins/Editor/WaveformEditor/Source/WaveformEditor/Private/WaveformEditorModule.cpp b/Engine/Plugins/Editor/WaveformEditor/Source/WaveformEditor/Private/WaveformEditorModule.cpp
--- a/Engine/Plugins/Editor/WaveformEditor/Source/WaveformEditor/Private/WaveformEditorModule.cpp
+++ b/Engine/Plugins/Editor/WaveformEditor/Source/WaveformEditor/Private/WaveformEditorModule.cpp
@@ -13,8 +13,42 @@
 #include "WaveformEditorInstantiator.h"
 #include "WaveformEditorLog.h"
 
+#include "UObject/PackageReload.h"
+
 DEFINE_LOG_CATEGORY(LogWaveformEditor);
 
+namespace
+{
+	void AnalyzeSoundWaveLoudness(UObject* Object)
+	{
+		if (Object && Object->IsA<USoundWave>())
+		{
+			TObjectPtr<USoundWave> SoundWave = CastChecked<USoundWave>(Object);
+
+			TArray<uint8> ImportedRawPCMData;
+			uint32 BufferSampleRate = 0;
+			uint16 ImportedNumChannels = 0;
+
+			Audio::FAlignedFloatBuffer AudioBuffer;
+
+			if (!SoundWave->GetImportedSoundWaveData(ImportedRawPCMData, BufferSampleRate, ImportedNumChannels) || BufferSampleRate == 0)
+			{
+				UE_LOG(LogWaveformEditor, Warning, TEXT("Soundwave loudness analyzation failed in Waveform Editor Module."));
+
+				return;
+			}
+
+			const int64 NumWaveformSamples = ImportedRawPCMData.Num() * sizeof(uint8) / sizeof(int16);
+			AudioBuffer.SetNumUninitialized(NumWaveformSamples);
+			Audio::ArrayPcm16ToFloat(MakeArrayView((int16*)ImportedRawPCMData.GetData(), NumWaveformSamples), AudioBuffer);
+
+			const float SamplePeakDB = WaveformAudioAnalysis::GetPeakSampleValue(AudioBuffer);
+			const float LUFS = WaveformAudioAnalysis::GetLUFS(AudioBuffer, BufferSampleRate, ImportedNumChannels);
+
+			SoundWave->SetLoudnessValues(LUFS, SamplePeakDB);
+		}
+	}
+}
 
 void FWaveformEditorModule::StartupModule()
 {
@@ -31,36 +65,36 @@ void FWaveformEditorModule::StartupModule()
 			{
 				if (TObjectPtr<UImportSubsystem> ImportSubsystem = GEditor->GetEditorSubsystem<UImportSubsystem>())
 				{
-					OnAssetPostImportHandle = ImportSubsystem->OnAssetPostImport.AddLambda([](UFactory* Factory, UObject* Object)
+					OnAssetPostImportHandle = ImportSubsystem->OnAssetPostImport.AddLambda([&](UFactory* Factory, UObject* Object)
 						{
-							if (Object && Object->IsA<USoundWave>())
-							{
-								TObjectPtr<USoundWave> SoundWave = CastChecked<USoundWave>(Object);
-
-								TArray<uint8> ImportedRawPCMData;
-								uint32 BufferSampleRate = 0;
-								uint16 ImportedNumChannels = 0;
+							AnalyzeSoundWaveLoudness(Object);
+						});
+				}
+			}
 
-								Audio::FAlignedFloatBuffer AudioBuffer;
+			OnPackageReloadedHandle = FCoreUObjectDelegates::OnPackageReloaded.AddLambda([](EPackageReloadPhase Phase, FPackageReloadedEvent* ReloadedEvent)
+				{
+					if (Phase == EPackageReloadPhase::PostPackageFixup)
+					{
+						check(ReloadedEvent);
 
-								if (!SoundWave->GetImportedSoundWaveData(ImportedRawPCMData, BufferSampleRate, ImportedNumChannels) || BufferSampleRate == 0)
-								{
-									UE_LOG(LogWaveformEditor, Warning, TEXT("Soundwave loudness analyzation failed on import."))
-										return;
-								}
+						const UPackage* ReloadedPackage = ReloadedEvent->GetNewPackage();
 
-								const int64 NumWaveformSamples = ImportedRawPCMData.Num() * sizeof(uint8) / sizeof(int16);
-								AudioBuffer.SetNumUninitialized(NumWaveformSamples);
-								Audio::ArrayPcm16ToFloat(MakeArrayView((int16*)ImportedRawPCMData.GetData(), NumWaveformSamples), AudioBuffer);
+						if (ReloadedPackage == nullptr)
+						{
+							return;
+						}
 
-								const float SamplePeakDB = WaveformAudioAnalysis::GetPeakSampleValue(AudioBuffer);
-								const float LUFS = WaveformAudioAnalysis::GetLUFS(AudioBuffer, BufferSampleRate, ImportedNumChannels);
+						TArray<UObject*> ObjectsInPackage;
+						constexpr bool bIncludeNestedObjects = true;
+						GetObjectsWithOuter(ReloadedPackage, ObjectsInPackage, bIncludeNestedObjects);
 
-								SoundWave->SetLoudnessValues(LUFS, SamplePeakDB);
-							}
-						});
-				}
-			}
+						for (UObject* Object : ObjectsInPackage)
+						{
+							AnalyzeSoundWaveLoudness(Object);
+						}
+					}
+				});
 		});
 }
 
@@ -77,6 +111,8 @@ void FWaveformEditorModule::ShutdownModule()
 			ImportSubsystem->OnAssetPostImport.Remove(OnAssetPostImportHandle);
 		}
 	}
+
+	FCoreUObjectDelegates::OnPackageReloaded.Remove(OnPackageReloadedHandle);
 }
 
 void FWaveformEditorModule::RegisterContentBrowserExtensions(IWaveformEditorInstantiator* Instantiator)
diff --git a/Engine/Plugins/Editor/WaveformEditor/Source/WaveformEditor/Public/WaveformEditorModule.h b/Engine/Plugins/Editor/WaveformEditor/Source/WaveformEditor/Public/WaveformEditorModule.h
--- a/Engine/Plugins/Editor/WaveformEditor/Source/WaveformEditor/Public/WaveformEditorModule.h
+++ b/Engine/Plugins/Editor/WaveformEditor/Source/WaveformEditor/Public/WaveformEditorModule.h
@@ -21,5 +21,6 @@ private:
 
 	FDelegateHandle OnPostEngineInitHandle;
 	FDelegateHandle OnAssetPostImportHandle;
+	FDelegateHandle OnPackageReloadedHandle;
 };
 
