[Waveform Editor] - Sample Peak zeroed out when using DWeighted Loudness mode

diff --git a/Engine/Plugins/Editor/WaveformEditor/Source/WaveformTransformations/Private/WaveformTransformationNormalize.cpp b/Engine/Plugins/Editor/WaveformEditor/Source/WaveformTransformations/Private/WaveformTransformationNormalize.cpp
--- a/Engine/Plugins/Editor/WaveformEditor/Source/WaveformTransformations/Private/WaveformTransformationNormalize.cpp
+++ b/Engine/Plugins/Editor/WaveformEditor/Source/WaveformTransformations/Private/WaveformTransformationNormalize.cpp
@@ -44,11 +44,16 @@ FWaveTransformationNormalize::FWaveTransformationNormalize(float InTarget, float
 
 	const float TargetOffset = Target - PeakDecibelValue;
 	const float MakeupGain = FMath::Clamp(TargetOffset, -MaxGain, MaxGain);
+	const float LinearPeakSample = Audio::ArrayMaxAbsValue(InputAudio);
 	InOutWaveInfo.StartFrameOffset = 0;
 
-	if (MakeupGain != 0.f)
+	if (MakeupGain != 0.f && LinearPeakSample > 0.f)
 	{
-		Audio::ArrayMultiplyByConstantInPlace(InputAudio, Audio::ConvertToLinear(MakeupGain));
+		// Prevent clipping from over amplification
+		const float MaxMakeupGain = 1.f / LinearPeakSample;
+		const float ClampedLinearMakeupGain = FMath::Min(Audio::ConvertToLinear(MakeupGain), MaxMakeupGain);
+
+		Audio::ArrayMultiplyByConstantInPlace(InputAudio, ClampedLinearMakeupGain);
 	}
 }
 
