Waveform Editor: Fix for right click getting stuck when scrolling waveform

diff --git a/Engine/Plugins/Editor/WaveformEditor/Source/WaveformEditorWidgets/Private/SWaveformTransformationsOverlay.cpp b/Engine/Plugins/Editor/WaveformEditor/Source/WaveformEditorWidgets/Private/SWaveformTransformationsOverlay.cpp
--- a/Engine/Plugins/Editor/WaveformEditor/Source/WaveformEditorWidgets/Private/SWaveformTransformationsOverlay.cpp
+++ b/Engine/Plugins/Editor/WaveformEditor/Source/WaveformEditorWidgets/Private/SWaveformTransformationsOverlay.cpp
@@ -97,8 +97,16 @@ void SWaveformTransformationsOverlay::UpdateLayerConstraints()
 
 FReply SWaveformTransformationsOverlay::OnMouseButtonDown(const FGeometry& MyGeometry, const FPointerEvent& MouseEvent)
 {
-	ZoomController->OnMouseButtonDown(MyGeometry, MouseEvent);
-	return AudioWidgetsUtils::RouteMouseInput(&SWidget::OnMouseButtonDown, MouseEvent, MakeArrayView(TransformationLayers.GetData(), TransformationLayers.Num()), false);
+	FReply InputReply = AudioWidgetsUtils::RouteMouseInput(&SWidget::OnMouseButtonDown, MouseEvent, MakeArrayView(TransformationLayers.GetData(), TransformationLayers.Num()), false);
+
+	// If the renderers do not consume input, then allow zoom scrolling
+	if (!InputReply.IsEventHandled())
+	{
+		check(ZoomController != nullptr);
+		InputReply = ZoomController->OnMouseButtonDown(MyGeometry, MouseEvent);
+	}
+
+	return InputReply;
 }
 
 FReply SWaveformTransformationsOverlay::OnMouseButtonUp(const FGeometry& MyGeometry, const FPointerEvent& MouseEvent)
