[Waveform Editor] - Fade-In/Out Functions in Toolbar do not affect Fade Function parameters

diff --git a/Engine/Plugins/Editor/WaveformEditor/Source/WaveformEditor/Private/WaveformEditor.cpp b/Engine/Plugins/Editor/WaveformEditor/Source/WaveformEditor/Private/WaveformEditor.cpp
--- a/Engine/Plugins/Editor/WaveformEditor/Source/WaveformEditor/Private/WaveformEditor.cpp
+++ b/Engine/Plugins/Editor/WaveformEditor/Source/WaveformEditor/Private/WaveformEditor.cpp
@@ -487,10 +487,17 @@ void FWaveformEditor::ToggleFadeIn()
 				CachedFadeInAmount = DefaultFadeInAmount;
 			}
 
+			// Ensures the selected subtype (linear, exponential, logarithmic, or sigmoid)
+			if (!TrimFadeTransformation->FadeFunctions.FadeIn.IsA(FadeInType.Get()))
+			{
+				TrimFadeTransformation->FadeFunctions.FadeIn = NewObject<UFadeFunction>(TrimFadeTransformation.Get(), FadeInType.Get(), NAME_None, RF_Transactional);
+			}
+
 			TrimFadeTransformation->FadeFunctions.FadeIn->Duration = CachedFadeInAmount;
 		}
 		else //if fade duration is not 0, assume there is a fade applied and toggle fade "off" by setting fade time to 0
 		{
+			// No need to check for fade type since we are toggling it "off". Type check is done during toggle "on"
 			CachedFadeInAmount = TrimFadeTransformation->FadeFunctions.FadeIn->Duration;
 
 			//No need to set FadeCurve when toggling off. This will preserve user curve data if they later manually change the StartFadeTime.
@@ -531,10 +538,17 @@ void FWaveformEditor::ToggleFadeOut()
 				CachedFadeOutAmount = DefaultFadeOutAmount;
 			}
 
+			// Ensures the selected subtype (linear, exponential, logarithmic, or sigmoid)
+			if (!TrimFadeTransformation->FadeFunctions.FadeOut.IsA(FadeOutType.Get()))
+			{
+				TrimFadeTransformation->FadeFunctions.FadeOut = NewObject<UFadeFunction>(TrimFadeTransformation.Get(), FadeOutType.Get(), NAME_None, RF_Transactional);
+			}
+
 			TrimFadeTransformation->FadeFunctions.FadeOut->Duration = CachedFadeOutAmount;
 		}
 		else //if fade duration is not 0, assume there is a fade applied and toggle fade "off" by setting fade time to 0
 		{
+			// No need to check for fade type since we are toggling it "off". Type check is done during toggle "on"
 			CachedFadeOutAmount = TrimFadeTransformation->FadeFunctions.FadeOut->Duration;
 
 			//No need to set FadeCurve when toggling off. This will preserve user curve data if they later manually change the StartFadeTime.
