[Waveform Editor] - Clarify tooltip for import/reimport

diff --git a/Engine/Plugins/Editor/WaveformEditor/Source/WaveformEditor/Private/WaveformEditorCommands.cpp b/Engine/Plugins/Editor/WaveformEditor/Source/WaveformEditor/Private/WaveformEditorCommands.cpp
--- a/Engine/Plugins/Editor/WaveformEditor/Source/WaveformEditor/Private/WaveformEditorCommands.cpp
+++ b/Engine/Plugins/Editor/WaveformEditor/Source/WaveformEditor/Private/WaveformEditorCommands.cpp
@@ -29,9 +29,9 @@ void FWaveformEditorCommands::RegisterCommands()
 	UI_COMMAND(ExportFormatStereo, "Stereo", "Sets the export format to stereo", EUserInterfaceActionType::RadioButton, FInputChord());
 
 	UI_COMMAND(ReimportAsset, "Reimport Asset", "Reimports the inspected sound wave uasset", EUserInterfaceActionType::Button, FInputChord());
-	UI_COMMAND(ReimportModeSameFile, "Reimport from same file", "Reimports asset from the same file", EUserInterfaceActionType::RadioButton, FInputChord());
-	UI_COMMAND(ReimportModeSameFileOverwriteTransformations, "Reimport from same file and overwrite transformations", "Reimports asset from the same file and overwrites the transformations", EUserInterfaceActionType::RadioButton, FInputChord());
-	UI_COMMAND(ReimportModeNewFile, "Reimport from new file", "Reimports asset from a selected file", EUserInterfaceActionType::RadioButton, FInputChord());
+	UI_COMMAND(ReimportModeSameFile, "Reimports asset via referenced .wav file", "Reimports asset from the original .wav file", EUserInterfaceActionType::RadioButton, FInputChord());
+	UI_COMMAND(ReimportModeSameFileOverwriteTransformations, "Reimport asset via referenced .wav file and overwrite transformations", "Reimport from the original .wav file and overwrite transformations", EUserInterfaceActionType::RadioButton, FInputChord());
+	UI_COMMAND(ReimportModeNewFile, "Reimport from a selected .wav file", "Reimports asset from a selected .wav file", EUserInterfaceActionType::RadioButton, FInputChord());
 
 	UI_COMMAND(ToggleFadeIn, "Toggle Fade In", "Toggle adding Fade In Transformation onto the sound wave", EUserInterfaceActionType::Button, FInputChord());
 	UI_COMMAND(FadeInLinear, "Linear Fade In", "Fade In Curve type applied when Toggle Fade In pressed", EUserInterfaceActionType::RadioButton, FInputChord());
