[Waveform Editor] - Cue Point positions not being recognized correctly when adjusting Trim times

diff --git a/Engine/Plugins/Editor/WaveformEditor/Source/WaveformTransformations/Private/WaveformTransformationMarkers.cpp b/Engine/Plugins/Editor/WaveformEditor/Source/WaveformTransformations/Private/WaveformTransformationMarkers.cpp
--- a/Engine/Plugins/Editor/WaveformEditor/Source/WaveformTransformations/Private/WaveformTransformationMarkers.cpp
+++ b/Engine/Plugins/Editor/WaveformEditor/Source/WaveformTransformations/Private/WaveformTransformationMarkers.cpp
@@ -14,6 +14,51 @@ namespace
 {
 	static const TArray<int32> SizeIncrements = { 1, 10, 50, 100 };
 	static const TArray<FKey> IndexKeyCodes = { EKeys::One, EKeys::Two, EKeys::Three, EKeys::Four, EKeys::Five, EKeys::Six, EKeys::Seven, EKeys::Eight, EKeys::Nine };
+
+	// Create a copy of the CuePoint data that is shifted by the StartFrameOffset created by transformations like TrimFade
+	const TArray<FSoundWaveCuePoint> GetAdjustedCuePoints(const TArray<FSoundWaveCuePoint>& InCuePoints, const int64 InStartFrameOffset, bool bRemoveTrimmed = false)
+	{
+		check(InStartFrameOffset >= 0);
+		if (InStartFrameOffset == 0)
+		{
+			return InCuePoints;
+		}
+
+		TArray<FSoundWaveCuePoint> FilteredAndAdjustedCuePoints;
+		FilteredAndAdjustedCuePoints.Reserve(InCuePoints.Num());
+		
+		for (const FSoundWaveCuePoint& Marker : InCuePoints)
+		{
+			const int64 NewPosition = Marker.FramePosition - InStartFrameOffset;
+
+			if (NewPosition < 0 && bRemoveTrimmed)
+			{
+				// Keep loop regions even if only the end frame is within the trimmed region
+				if (!Marker.IsLoopRegion() || Marker.FrameLength + NewPosition <= 0)
+				{
+					continue; // Skip adding this marker to filtered list
+				}
+			}
+
+			FSoundWaveCuePoint FilteredMarker = Marker;
+
+			// If loop region is cut, resize it to maintain proper relative loop end point
+			if (FilteredMarker.IsLoopRegion())
+			{
+				if (FilteredMarker.FrameLength + NewPosition > 0)
+				{
+					FilteredMarker.FrameLength += NewPosition;
+				}
+
+				FilteredMarker.FrameLength = FMath::Max(FilteredMarker.FrameLength, UWaveCueArray::MinLoopSize);
+			}
+
+			FilteredMarker.FramePosition = FMath::Max(NewPosition, 0);
+			FilteredAndAdjustedCuePoints.Add(FilteredMarker);
+		}
+		
+		return FilteredAndAdjustedCuePoints;
+	}
 }
 
 FWaveTransformationMarkers::FWaveTransformationMarkers(double InStartLoopTime, double InEndLoopTime)
@@ -623,10 +668,6 @@ bool UWaveformTransformationMarkers::AdjustLoopPreviewIfNotAligned()
 #if WITH_EDITOR
 void UWaveformTransformationMarkers::OverwriteSoundWaveData(USoundWave& InOutSoundWave)
 {
-	// Overwriting soundwave data can cause a change in number of samples, invalidating the FramePositions of CuesAndLoops
-	// Subtracting StartFrameOffset from FramePosition shifts the CuesAndLoops to the correct relative position
-	TArray<FSoundWaveCuePoint> FrameAdjustedCuesAndLoops = Markers->CuesAndLoops;
-
 	// Users can export a loop region, shift FramePositions relative to the Loop region
 	if (bIsPreviewingLoopRegion)
 	{
@@ -635,23 +676,9 @@ void UWaveformTransformationMarkers::OverwriteSoundWaveData(USoundWave& InOutSou
 		StartFrameOffset = StartFrameOffset < StartLoopFramePosition ? StartLoopFramePosition : StartFrameOffset;
 	}
 
-	if (StartFrameOffset != 0)
-	{
-		for (FSoundWaveCuePoint& Marker : FrameAdjustedCuesAndLoops)
-		{
-			const int64 NewPosition = Marker.FramePosition - StartFrameOffset;
-
-			// If loop region is cut, resize it to maintain proper relative loop end point
-			// else maintain original loop region size
-			if (NewPosition < 0 && Marker.IsLoopRegion())
-			{
-				Marker.FrameLength = Marker.FrameLength + NewPosition > 0 ? Marker.FrameLength + NewPosition : Marker.FrameLength;
-				Marker.FrameLength = FMath::Max(Marker.FrameLength, UWaveCueArray::MinLoopSize);
-			}
-
-			Marker.FramePosition = FMath::Max(NewPosition, 0);
-		}
-	}
+	// Overwriting soundwave data can cause a change in number of samples, invalidating the FramePositions of CuesAndLoops
+	// Subtracting StartFrameOffset from FramePosition shifts the CuesAndLoops to the correct relative position
+	const TArray<FSoundWaveCuePoint>& FrameAdjustedCuesAndLoops = GetAdjustedCuePoints(Markers->CuesAndLoops, StartFrameOffset);
 
 	InOutSoundWave.SetSoundWaveCuePoints(FrameAdjustedCuesAndLoops);
 }
@@ -660,6 +687,9 @@ void UWaveformTransformationMarkers::GetTransformationInfo(FWaveformTransformati
 {
 	check(Markers != nullptr);
 	InOutTransformationInfo.AllCuePoints.Append(Markers->CuesAndLoops);
+
+	constexpr bool bRemoveTrimmed = true;
+	InOutTransformationInfo.AllCuePointsRelativeToStartTime.Append(GetAdjustedCuePoints(Markers->CuesAndLoops, StartFrameOffset, bRemoveTrimmed));
 }
 
 FSoundWaveCuePoint* UWaveformTransformationMarkers::GetSelectedMarker() const
diff --git a/Engine/Source/Runtime/AudioExtensions/Public/IWaveformTransformation.h b/Engine/Source/Runtime/AudioExtensions/Public/IWaveformTransformation.h
--- a/Engine/Source/Runtime/AudioExtensions/Public/IWaveformTransformation.h
+++ b/Engine/Source/Runtime/AudioExtensions/Public/IWaveformTransformation.h
@@ -118,6 +118,7 @@ struct FWaveTransformUObjectConfiguration
 struct FWaveformTransformationInfo
 {
 	TArray<FSoundWaveCuePoint> AllCuePoints; // Cue Points and Loop Regions
+	TArray<FSoundWaveCuePoint> AllCuePointsRelativeToStartTime; // Cue Points and Loop Regions offset by start time
 };
 #endif //WITH_EDITORONLY_DATA
 
diff --git a/Engine/Source/Runtime/Engine/Private/SoundWave.cpp b/Engine/Source/Runtime/Engine/Private/SoundWave.cpp
--- a/Engine/Source/Runtime/Engine/Private/SoundWave.cpp
+++ b/Engine/Source/Runtime/Engine/Private/SoundWave.cpp
@@ -2711,7 +2711,7 @@ TArray<FSoundWaveCuePoint> USoundWave::GetCuePointsScaledForSampleRate(const flo
 		CuePointsScaled = CuePoints;
 		break;
 	case ESoundWaveCuePointOrigin::MarkerTransformation:
-		CuePointsScaled = GetWaveformTransformationInfo().AllCuePoints;
+		CuePointsScaled = GetWaveformTransformationInfo().AllCuePointsRelativeToStartTime;
 		break;
 	default:
 		CuePointsScaled = CuePoints;
@@ -5127,7 +5127,7 @@ void USoundWave::SerializeCuePoints(FArchive& Ar, const bool bIsLoadingFromCooke
 			PlatformCuePoints = CuePoints;
 			break;
 		case ESoundWaveCuePointOrigin::MarkerTransformation:
-			PlatformCuePoints = GetWaveformTransformationInfo().AllCuePoints;
+			PlatformCuePoints = GetWaveformTransformationInfo().AllCuePointsRelativeToStartTime;
 			break;
 		default:
 			PlatformCuePoints = CuePoints;
