Waveform Editor: nudging bughawk fixes and loop selection fix

diff --git a/Engine/Plugins/Editor/WaveformEditor/Source/WaveformTransformations/Private/WaveformTransformationMarkers.cpp b/Engine/Plugins/Editor/WaveformEditor/Source/WaveformTransformations/Private/WaveformTransformationMarkers.cpp
--- a/Engine/Plugins/Editor/WaveformEditor/Source/WaveformTransformations/Private/WaveformTransformationMarkers.cpp
+++ b/Engine/Plugins/Editor/WaveformEditor/Source/WaveformTransformations/Private/WaveformTransformationMarkers.cpp
@@ -193,7 +193,7 @@ void UWaveformTransformationMarkers::UpdateConfiguration(FWaveTransformUObjectCo
 	SampleRate = InOutConfiguration.SampleRate;
 	AvailableWaveformDuration = InOutDuration;
 	NumChannels = InOutConfiguration.NumChannels;
-	NumAvailableSamples = static_cast<int64>(InOutDuration * SampleRate);
+	NumAvailableSamples = static_cast<int64>(InOutDuration * SampleRate * NumChannels);
 
 	if (!bHasLoggedCutByLoop)
 	{
@@ -268,8 +268,9 @@ void UWaveformTransformationMarkers::ModifyMarkerLoopRegion(ELoopModificationCon
 	check(NumChannels > 0);
 	check(NumAvailableSamples > 0);
 	check(SizeIncrements.Num() > 0);
+	check(SelectedIncrement < SizeIncrements.Num() && SelectedIncrement >= 0);
 
-	int64 FramesToShift = static_cast<int64>((SizeIncrements[SelectedIncrement] / 1000.0f) * SampleRate);
+	const int64 FramesToShift = static_cast<int64>((SizeIncrements[SelectedIncrement] / 1000.0f) * SampleRate);
 
 	if (GEditor && GEditor->Trans)
 	{
@@ -319,12 +320,13 @@ void UWaveformTransformationMarkers::ModifyMarkerLoopRegion(ELoopModificationCon
 			--SelectedIncrement;
 		}
 	}
+	break;
 	case ELoopModificationControls::LeftHandleNextZeroCrossing:
 	{
-		const int64 StartFrame = (LoopRegion->FramePosition) * NumChannels;
-		const int64 MaxFrames = 0;
+		const int64 StartSample = (LoopRegion->FramePosition) * NumChannels;
+		const int64 EndSample = 0;
 
-		const int64 FramesToMove = GetFramesToNextZeroCrossing(StartFrame, MaxFrames, PCMArray, false, NoiseThreshold);
+		const int64 FramesToMove = GetFramesToNextZeroCrossing(StartSample, EndSample, PCMArray, false, NoiseThreshold);
 
 		LoopRegion->FramePosition -= FramesToMove;
 		LoopRegion->FrameLength += FramesToMove;
@@ -332,10 +334,10 @@ void UWaveformTransformationMarkers::ModifyMarkerLoopRegion(ELoopModificationCon
 	break;
 	case ELoopModificationControls::LeftHandlePreviousZeroCrossing:
 	{
-		const int64 StartFrame = (LoopRegion->FramePosition) * NumChannels;
-		const int64 MaxFrames = (LoopRegion->FramePosition + LoopRegion->FrameLength) * NumChannels;
+		const int64 StartSample = (LoopRegion->FramePosition) * NumChannels;
+		const int64 EndSample = (LoopRegion->FramePosition + LoopRegion->FrameLength) * NumChannels;
 
-		const int64 FramesToMove = GetFramesToNextZeroCrossing(StartFrame, MaxFrames, PCMArray, true, NoiseThreshold);
+		const int64 FramesToMove = GetFramesToNextZeroCrossing(StartSample, EndSample, PCMArray, true, NoiseThreshold);
 
 		LoopRegion->FramePosition += FramesToMove;
 		LoopRegion->FrameLength -= FramesToMove;
@@ -343,24 +345,24 @@ void UWaveformTransformationMarkers::ModifyMarkerLoopRegion(ELoopModificationCon
 	break;
 	case ELoopModificationControls::RightHandleNextZeroCrossing:
 	{
-		const int64 StartFrame = (LoopRegion->FramePosition + LoopRegion->FrameLength) * NumChannels;
-		const int64 MaxFrames = NumAvailableSamples;
+		const int64 StartSample = (LoopRegion->FramePosition + LoopRegion->FrameLength) * NumChannels;
+		const int64 EndSample = NumAvailableSamples;
 
-		LoopRegion->FrameLength += GetFramesToNextZeroCrossing(StartFrame, MaxFrames, PCMArray, true, NoiseThreshold);
+		LoopRegion->FrameLength += GetFramesToNextZeroCrossing(StartSample, EndSample, PCMArray, true, NoiseThreshold);
 	}
 	break;
 	case ELoopModificationControls::RightHandlePreviousZeroCrossing:
 	{
-		const int64 StartFrame = (LoopRegion->FramePosition + LoopRegion->FrameLength) * NumChannels;
-		const int64 MaxFrames = LoopRegion->FramePosition;
+		const int64 StartSample = (LoopRegion->FramePosition + LoopRegion->FrameLength) * NumChannels;
+		const int64 EndSample = LoopRegion->FramePosition * NumChannels;
 
-		LoopRegion->FrameLength -= GetFramesToNextZeroCrossing(StartFrame, MaxFrames, PCMArray, false, NoiseThreshold);
+		LoopRegion->FrameLength -= GetFramesToNextZeroCrossing(StartSample, EndSample, PCMArray, false, NoiseThreshold);
 	}
 	default:
 		break;
 	}
 
-	int64 MaxTotalFrames = static_cast<int64>(NumAvailableSamples) / static_cast<int64>(NumChannels);
+	const int64 MaxTotalFrames = static_cast<int64>(NumAvailableSamples) / static_cast<int64>(NumChannels);
 	check(MaxTotalFrames > 0);
 	// Prevent the frame position from exceeding the wave length
 	LoopRegion->FramePosition = FMath::Clamp(LoopRegion->FramePosition, 0, MaxTotalFrames - 1);
@@ -442,7 +444,7 @@ void UWaveformTransformationMarkers::CycleMarkerLoopRegion(ELoopModificationCont
 	OnTransformationChanged.ExecuteIfBound(false);
 }
 
-int64 UWaveformTransformationMarkers::GetFramesToNextZeroCrossing(int64 StartFrame, int64 EndFrame, const TArray<float>& PCMArray, bool bCheckIncrementing, float NoiseThreshold) const
+int64 UWaveformTransformationMarkers::GetFramesToNextZeroCrossing(int64 StartSample, int64 EndSample, const TArray<float>& PCMArray, bool bCheckIncrementing, float NoiseThreshold) const
 {
 	check(NumChannels > 0);
 	int64 FramesMoved = 0;
@@ -453,7 +455,7 @@ int64 UWaveformTransformationMarkers::GetFramesToNextZeroCrossing(int64 StartFra
 
 	if (bCheckIncrementing)
 	{
-		for (int32 i = StartFrame; i < EndFrame; i += NumChannels)
+		for (int64 i = StartSample; i < EndSample; i += NumChannels)
 		{
 			if (CheckZeroCrossingFrame(i, PCMArray, bAllFramesPositiveLastFrame, bAllFramesNegativeLastFrame, bAllFramesLowNoiseLastFrame, NoiseThreshold))
 			{
@@ -465,7 +467,7 @@ int64 UWaveformTransformationMarkers::GetFramesToNextZeroCrossing(int64 StartFra
 	}
 	else
 	{
-		for (int32 i = StartFrame; i >= EndFrame; i -= NumChannels)
+		for (int64 i = StartSample; i >= EndSample; i -= NumChannels)
 		{
 			if (CheckZeroCrossingFrame(i, PCMArray, bAllFramesPositiveLastFrame, bAllFramesNegativeLastFrame, bAllFramesLowNoiseLastFrame, NoiseThreshold))
 			{
@@ -482,7 +484,12 @@ int64 UWaveformTransformationMarkers::GetFramesToNextZeroCrossing(int64 StartFra
 int64 UWaveformTransformationMarkers::CheckZeroCrossingFrame(int64 CurrentFrame, const TArray<float>& PCMArray, bool& bAllFramesPositiveLastFrame, bool& bAllFramesNegativeLastFrame, bool& bAllFramesLowNoiseLastFrame, float NoiseThreshold) const
 {
 	check(NumChannels > 0);
-	check(CurrentFrame < PCMArray.Num());
+	check(CurrentFrame % NumChannels == 0); // Check for frame alignment
+
+	if (CurrentFrame + NumChannels > PCMArray.Num())
+	{
+		return 0;
+	}
 
 	int32 PositiveSampleCount = 0;
 	int32 NegativeSampleCount = 0;
@@ -717,7 +724,7 @@ void UWaveformTransformationMarkers::PreviewSelectedLoop()
 	check(SampleRate > 0);
 	check(NumChannels > 0);
 
-	const int64 EndFramePosition = StartFrameOffset + NumAvailableSamples / NumChannels;
+	const int64 EndFramePosition = StartFrameOffset + (NumAvailableSamples / NumChannels);
 
 	for (FSoundWaveCuePoint CuePoint : Markers->CuesAndLoops)
 	{
