[Waveform Editor] - Metadata does not update after reloading a Sound Wave

diff --git a/Engine/Plugins/Editor/WaveformEditor/Source/WaveformEditor/Private/WaveformEditorModule.cpp b/Engine/Plugins/Editor/WaveformEditor/Source/WaveformEditor/Private/WaveformEditorModule.cpp
--- a/Engine/Plugins/Editor/WaveformEditor/Source/WaveformEditor/Private/WaveformEditorModule.cpp
+++ b/Engine/Plugins/Editor/WaveformEditor/Source/WaveformEditor/Private/WaveformEditorModule.cpp
@@ -46,6 +46,7 @@ namespace
 			const float LUFS = WaveformAudioAnalysis::GetLUFS(AudioBuffer, BufferSampleRate, ImportedNumChannels);
 
 			SoundWave->SetLoudnessValues(LUFS, SamplePeakDB);
+			SoundWave->PostEditChange();
 		}
 	}
 }
