[Waveform Editor] - Editor crashes when Trim Fade goes across a Cue Point in the Waveform editor

diff --git a/Engine/Plugins/Editor/WaveformEditor/Source/WaveformTransformations/Private/WaveformTransformationMarkers.cpp b/Engine/Plugins/Editor/WaveformEditor/Source/WaveformTransformations/Private/WaveformTransformationMarkers.cpp
--- a/Engine/Plugins/Editor/WaveformEditor/Source/WaveformTransformations/Private/WaveformTransformationMarkers.cpp
+++ b/Engine/Plugins/Editor/WaveformEditor/Source/WaveformTransformations/Private/WaveformTransformationMarkers.cpp
@@ -96,7 +96,7 @@ void FWaveTransformationMarkers::ProcessAudio(Audio::FWaveformTransformationWave
 	const int64 FirstSampleOfLastFrame = LastInputAudioIndex - NumChannelsMinusOne; // The last sample that begins a frame
 	int64 StartSample = FMath::Min(FMath::FloorToInt64(StartLoopTime * InOutWaveInfo.SampleRate) * InOutWaveInfo.NumChannels, FirstSampleOfLastFrame);
 	int64 EndSample = LastInputAudioIndex;
-	check(StartSample < LastInputAudioIndex);
+	check(StartSample <= LastInputAudioIndex);
 
 	if (EndLoopTime > 0.f)
 	{
