Playhead stutters when adjusting Trim values
UE-324789
Redoing Start Time adjustment can displace the Scrubber
UE-306846
Make transport state more consitent through different change types via FWaveformEditor::NotifyPostChange

Jitter error magnitude varied with the size of an interactive change -- if one dragged the trimfade extents quickly, the playhead would jitter farther.  Behaviour was different re typing in values (a valueset change) and dragging them to the same value (an interactive change), thus something more than just rounding was at play.  Long story short, this change makes them behave in the same, more desirable way.

diff --git a/Engine/Plugins/Editor/WaveformEditor/Source/WaveformEditor/Private/WaveformEditor.cpp b/Engine/Plugins/Editor/WaveformEditor/Source/WaveformEditor/Private/WaveformEditor.cpp
--- a/Engine/Plugins/Editor/WaveformEditor/Source/WaveformEditor/Private/WaveformEditor.cpp
+++ b/Engine/Plugins/Editor/WaveformEditor/Source/WaveformEditor/Private/WaveformEditor.cpp
@@ -1737,15 +1737,10 @@ void FWaveformEditor::NotifyPostChange(const FPropertyChangedEvent& PropertyChan
 
 	BindDelegatesForActiveTransformations();
 
-	if ((PropertyChangedEvent.ChangeType & EPropertyChangeType::ValueSet) == EPropertyChangeType::ValueSet
-	 && (PropertyChangedEvent.ChangeType & EPropertyChangeType::Interactive) != EPropertyChangeType::Interactive)
-	{
-		UpdateTransportState();
-
-		bIsInteractingWithTransformations = false;
+	UpdateTransportState();
+	bIsInteractingWithTransformations = false;
 
-		PromptToChangeCuePointOrigin();
-	}
+	PromptToChangeCuePointOrigin();
 }
 
 void FWaveformEditor::UpdateTransportState()
