[Mac only] Editor Crashes when exporting a Sound Wave

diff --git a/Engine/Plugins/Editor/WaveformEditor/Source/WaveformEditor/Private/WaveformEditorWaveWriter.cpp b/Engine/Plugins/Editor/WaveformEditor/Source/WaveformEditor/Private/WaveformEditorWaveWriter.cpp
--- a/Engine/Plugins/Editor/WaveformEditor/Source/WaveformEditor/Private/WaveformEditorWaveWriter.cpp
+++ b/Engine/Plugins/Editor/WaveformEditor/Source/WaveformEditor/Private/WaveformEditorWaveWriter.cpp
@@ -45,16 +45,28 @@ void FWaveformEditorWaveWriter::ExportTransformedWaveform(bool bPromptUserToSave
 		UE_LOG(LogWaveformEditor, Log, TEXT("Finished Exporting edited soundwave %s/%s"), *AssetPath, *AssetName);
 		TArray<UPackage*> PackagesToSave;
 
-		if(ResultingWave->GetPackage())
+		if (ResultingWave == nullptr)
 		{
-			PackagesToSave.Add(ResultingWave->GetPackage());
+			UE_LOG(LogWaveformEditor, Warning, TEXT("Resulting Exported soundwave is invalid: %s/%s"), *AssetPath, *AssetName);
 
-			FEditorFileUtils::EPromptReturnCode Result = FEditorFileUtils::PromptForCheckoutAndSave(PackagesToSave, false /*bCheckDirty*/, bPromptUserToSave);
+			return;
+		}
+
+		if (UPackage* SoundWavePackage = ResultingWave->GetPackage())
+		{
+			PackagesToSave.Add(SoundWavePackage);
 
-			// If the user cancels the checkout and save, delete the new generated soundwave
-			if (bPromptUserToSave && Result == FEditorFileUtils::PR_Cancelled)
+			SoundWavePackage->SetDirtyFlag(true);
+			const FString PackageFileName = FPackageName::LongPackageNameToFilename(SoundWavePackage->GetName(), FPackageName::GetAssetPackageExtension());
+			const bool bSuccess = UEditorLoadingAndSavingUtils::SavePackages(PackagesToSave, false);
+
+			if (bSuccess)
+			{
+				UE_LOG(LogWaveformEditor, Log, TEXT("Exported soundwave saving succeeded: % s / % s"), *AssetPath, *AssetName);
+			}
+			else
 			{
-				UEditorAssetLibrary::DeleteAsset(ResultingWave->GetPathName());
+				UE_LOG(LogWaveformEditor, Warning, TEXT("Exported soundwave saving failed: %s/%s"), *AssetPath, *AssetName);
 			}
 		}
 	};
