UE-318566
Address Bughawk feedback, WaveformTransformationMarkers.cpp

diff --git a/Engine/Plugins/Editor/WaveformEditor/Source/WaveformTransformations/Private/WaveformTransformationMarkers.cpp b/Engine/Plugins/Editor/WaveformEditor/Source/WaveformTransformations/Private/WaveformTransformationMarkers.cpp
--- a/Engine/Plugins/Editor/WaveformEditor/Source/WaveformTransformations/Private/WaveformTransformationMarkers.cpp
+++ b/Engine/Plugins/Editor/WaveformEditor/Source/WaveformTransformations/Private/WaveformTransformationMarkers.cpp
@@ -813,11 +813,13 @@ void UWaveCueArray::PostEditChangeChainProperty(struct FPropertyChangedChainEven
 										   |EPropertyChangeType::ArrayMove
 										   |EPropertyChangeType::Redirected) )
 	{
-		if (CuesAndLoops.Num() > 0)
+		if (CuesAndLoops.Num() > 0 && PropertyChangedEvent.PropertyChain.GetActiveMemberNode() != nullptr)
 		{
 			int32 MemberNodeIndex = PropertyChangedEvent.GetArrayIndex(PropertyChangedEvent.PropertyChain.GetActiveMemberNode()->GetValue()->GetName());
-			check(MemberNodeIndex < CuesAndLoops.Num());
-			SelectedCue = CuesAndLoops[MemberNodeIndex].CuePointID;
+			if (MemberNodeIndex >= 0 && MemberNodeIndex < CuesAndLoops.Num())
+			{
+				SelectedCue = CuesAndLoops[MemberNodeIndex].CuePointID;
+			}
 		}
 	}
 }
