Fixed crash when importing a USD file with the USDImporter into the content browser with the Waveform Editor plugin enabled.

#rnx

diff --git a/Engine/Plugins/Editor/WaveformEditor/Source/WaveformEditor/Private/WaveformEditorModule.cpp b/Engine/Plugins/Editor/WaveformEditor/Source/WaveformEditor/Private/WaveformEditorModule.cpp
--- a/Engine/Plugins/Editor/WaveformEditor/Source/WaveformEditor/Private/WaveformEditorModule.cpp
+++ b/Engine/Plugins/Editor/WaveformEditor/Source/WaveformEditor/Private/WaveformEditorModule.cpp
@@ -33,7 +33,7 @@ void FWaveformEditorModule::StartupModule()
 				{
 					OnAssetPostImportHandle = ImportSubsystem->OnAssetPostImport.AddLambda([](UFactory* Factory, UObject* Object)
 						{
-							if (Object->IsA<USoundWave>())
+							if (Object && Object->IsA<USoundWave>())
 							{
 								TObjectPtr<USoundWave> SoundWave = CastChecked<USoundWave>(Object);
 
