Fix for Negative INFINITY msvc compile issue on some machines

diff --git a/Engine/Plugins/Editor/WaveformEditor/Source/WaveformTransformations/Private/WaveformAudioAnalysisFunctions.cpp b/Engine/Plugins/Editor/WaveformEditor/Source/WaveformTransformations/Private/WaveformAudioAnalysisFunctions.cpp
--- a/Engine/Plugins/Editor/WaveformEditor/Source/WaveformTransformations/Private/WaveformAudioAnalysisFunctions.cpp
+++ b/Engine/Plugins/Editor/WaveformEditor/Source/WaveformTransformations/Private/WaveformAudioAnalysisFunctions.cpp
@@ -6,9 +6,15 @@
 #include "Algo/MaxElement.h"
 #include "LKFSNRTFactory.h"
 #include "MeterFactory.h"
+#include <limits>
 
 namespace WaveformAudioAnalysis
 {
+	constexpr float DbMin = -std::numeric_limits<float>::infinity();
+
+	// GatedLoudness starts analysis with a -70dB gate, if something falls below the gated value, we should return DbMin
+	constexpr float NoiseFloor = -70.f;
+
 	float GetRMSPeak(const Audio::FAlignedFloatBuffer& InputAudio, float SampleRate, int32 NumChannels)
 	{
 		check(NumChannels > 0);
@@ -32,6 +38,11 @@ namespace WaveformAudioAnalysis
 
 			if (MeterValue)
 			{
+				if (*MeterValue < NoiseFloor)
+				{
+					return DbMin;
+				}
+
 				return *MeterValue;
 			}
 		}
@@ -43,7 +54,7 @@ namespace WaveformAudioAnalysis
 
 			if (AnalyzerWorker == nullptr)
 			{
-				return 0.f;
+				return DbMin;
 			}
 
 			AnalyzerWorker->Analyze(MakeArrayView(InputAudio), Result.Get());
@@ -60,12 +71,17 @@ namespace WaveformAudioAnalysis
 
 				if (MeterEntry)
 				{
+					if (MeterEntry->MeterValue < NoiseFloor)
+					{
+						return DbMin;
+					}
+
 					return MeterEntry->MeterValue;
 				}
 			}
 		}
 
-		return 0.f;
+		return DbMin;
 	}
 
 	float GetLoudnessPeak(const Audio::FAlignedFloatBuffer& InputAudio, float SampleRate, int32 NumChannels)
@@ -81,7 +97,7 @@ namespace WaveformAudioAnalysis
 
 		if (AnalyzerWorker == nullptr)
 		{
-			return 0.f;
+			return DbMin;
 		}
 
 		AnalyzerWorker->Analyze(MakeArrayView(InputAudio), Result.Get());
@@ -98,18 +114,29 @@ namespace WaveformAudioAnalysis
 
 			if (LoudestEntry)
 			{
+				if (LoudestEntry->Loudness < NoiseFloor)
+				{
+					return DbMin;
+				}
+
 				return LoudestEntry->Loudness;
 			}
 		}
 
-		return 0.f;
+		return DbMin;
 	}
 
 	float GetPeakSampleValue(const Audio::FAlignedFloatBuffer& InputAudio)
 	{
 		const float MaxValue = Audio::ArrayMaxAbsValue(InputAudio);
+		float DbValue = Audio::ConvertToDecibels(MaxValue);
+
+		if (DbValue < NoiseFloor)
+		{
+			DbValue = DbMin;
+		}
 
-		return Audio::ConvertToDecibels(MaxValue);
+		return DbValue;
 	}
 
 	float GetLUFS(const Audio::FAlignedFloatBuffer& InputAudio, float SampleRate, int32 NumChannels)
@@ -123,7 +150,7 @@ namespace WaveformAudioAnalysis
 
 		if (Worker == nullptr)
 		{
-			return -INFINITY;
+			return DbMin;
 		}
 
 		Worker->Analyze(MakeArrayView(InputAudio), Result.Get());
@@ -135,9 +162,15 @@ namespace WaveformAudioAnalysis
 		{
 			const float GatedLoudness = LKFSResult->FindAggregateLoudnessStats(LKFSResult->ChannelIndexOverall)->GatedLoudness;
 
+			if (GatedLoudness == FLKFSNRTAggregateStats::InvalidLoudness || GatedLoudness < NoiseFloor)
+			{
+				// Sample not large enough to analyze
+				return DbMin;
+			}
+
 			return GatedLoudness;
 		}
 
-		return -INFINITY;
+		return DbMin;
 	}
 }
