Waveform Editor: Additional unit tests

diff --git a/Engine/Plugins/Editor/WaveformEditor/Source/WaveformEditor/Public/WaveformEditor.h b/Engine/Plugins/Editor/WaveformEditor/Source/WaveformEditor/Public/WaveformEditor.h
--- a/Engine/Plugins/Editor/WaveformEditor/Source/WaveformEditor/Public/WaveformEditor.h
+++ b/Engine/Plugins/Editor/WaveformEditor/Source/WaveformEditor/Public/WaveformEditor.h
@@ -54,6 +54,12 @@ public:
 
 	UE_API void ExportWaveform(bool bPromptUserToSave);
 
+	// Skips the playhead to the next marker
+	UE_API void SkipToNextMarker();
+	UE_API void ResetPlaybackToStart();
+
+	TSharedPtr<class FSparseSampledSequenceTransportCoordinator> GetTransportCoordinator() { return TransportCoordinator; }
+
 	static UE_API const double EndTimeInvalid;
 
 private:
@@ -64,7 +70,6 @@ private:
 	UE_API bool SetUpAssetReimport();
 	UE_API void ExecuteReimport();
 	UE_API void ExecuteOverwriteTransformations();
-	UE_API void ResetPlaybackToStart();
 
 	/**	Sets the wave editor layout */
 	UE_API const TSharedRef<FTabManager::FLayout> SetupStandaloneLayout();
@@ -129,8 +134,6 @@ private:
 	UE_API bool CanFadeOut();
 	UE_API void CreateMarker(bool bIsLoopRegion);
 	UE_API void DeleteMarker();
-	// Skips the playhead to the next marker
-	UE_API void SkipToNextMarker(); 
 	UE_API void RegenerateTransformations();
 	UE_API void UpdateTransformations(bool bMarkDirty = true);
 	UE_API void UpdateRenderElements();
