[Waveform] Additional unit test cases:
- Verify marker hit at correct time when playing a MetaSound playing a SoundWave.
- Test for handling Bad Chunk audio analysis
Fixed Loop Region Frame Length is increased in trimmed and exported Sound Wave so test would pass

#code-analyzed
#virtualized

diff --git a/Engine/Plugins/Editor/WaveformEditor/Source/WaveformTransformations/Private/WaveformTransformationMarkers.cpp b/Engine/Plugins/Editor/WaveformEditor/Source/WaveformTransformations/Private/WaveformTransformationMarkers.cpp
--- a/Engine/Plugins/Editor/WaveformEditor/Source/WaveformTransformations/Private/WaveformTransformationMarkers.cpp
+++ b/Engine/Plugins/Editor/WaveformEditor/Source/WaveformTransformations/Private/WaveformTransformationMarkers.cpp
@@ -42,10 +42,10 @@ namespace
 
 			FSoundWaveCuePoint FilteredMarker = Marker;
 
-			// If loop region is cut, resize it to maintain proper relative loop end point
 			if (FilteredMarker.IsLoopRegion())
 			{
-				if (FilteredMarker.FrameLength + NewPosition > 0)
+				// If loop region is cut, resize it to maintain proper relative loop end point
+				if (NewPosition < 0 && FilteredMarker.FrameLength + NewPosition > 0)
 				{
 					FilteredMarker.FrameLength += NewPosition;
 				}
diff --git a/Engine/Source/Runtime/AudioExtensions/Public/IWaveformTransformation.h b/Engine/Source/Runtime/AudioExtensions/Public/IWaveformTransformation.h
--- a/Engine/Source/Runtime/AudioExtensions/Public/IWaveformTransformation.h
+++ b/Engine/Source/Runtime/AudioExtensions/Public/IWaveformTransformation.h
@@ -83,6 +83,17 @@ struct FSoundWaveCuePoint
 	bool IsLoopRegion() const { return bIsLoopRegion; }
 	void SetLoopRegion(bool value) { bIsLoopRegion = value; }
 
+	const bool operator==(const FSoundWaveCuePoint& Other) const
+	{ 
+		return CuePointID == Other.CuePointID &&
+			FrameLength == Other.FrameLength &&
+			FramePosition == Other.FramePosition &&
+			Label == Other.Label &&
+			IsLoopRegion() == Other.IsLoopRegion();
+	}
+
+	const bool operator!=(const FSoundWaveCuePoint& Other) const { return !(*this == Other); }
+
 #if WITH_EDITORONLY_DATA
 	void ScaleFrameValues(float Factor)
 	{
