Waveform Editor: Snap to zero crossing shortcut
Integrated from 39.00

diff --git a/Engine/Plugins/Editor/WaveformEditor/Source/WaveformEditor/Private/WaveformEditor.cpp b/Engine/Plugins/Editor/WaveformEditor/Source/WaveformEditor/Private/WaveformEditor.cpp
--- a/Engine/Plugins/Editor/WaveformEditor/Source/WaveformEditor/Private/WaveformEditor.cpp
+++ b/Engine/Plugins/Editor/WaveformEditor/Source/WaveformEditor/Private/WaveformEditor.cpp
@@ -689,6 +689,38 @@ void FWaveformEditor::SkipToNextMarker()
 	}
 }
 
+void FWaveformEditor::SkipToNextZeroCrossing()
+{
+	check(SoundWave != nullptr);
+
+	const int64 StartSample = static_cast<int64>(TransportCoordinator->GetFocusPoint() * SoundWave->TotalSamples);
+	TArray<float> PCMArray = WaveformView.DataProvider->GetTransformedPCMData();
+
+	const float StartingFocusPoint = TransportCoordinator->GetFocusPoint();
+
+	const UWaveformEditorTransformationsSettings* TransformationsSettings = GetWaveformEditorTransformationsSettings();
+	check(TransformationsSettings);
+
+	const float NoiseThresholdValue = powf(10.0f, TransformationsSettings->NoiseThreshold / 20.0f);
+
+	const int64 FramesToMove = UWaveformTransformationMarkers::GetFramesToNextZeroCrossing(StartSample, SoundWave->TotalSamples, PCMArray, true, NoiseThresholdValue, SoundWave->NumChannels);
+
+	float ActiveDuration = TransformationChainConfig.EndTime - TransformationChainConfig.StartTime;
+	if (ActiveDuration <= 0.0f)
+	{
+		ActiveDuration = SoundWave->Duration;
+	}
+
+	check(ActiveDuration);
+	check(TransformationChainConfig.SampleRate > 0);
+
+	const float ActiveDurationInFrames = ActiveDuration * TransformationChainConfig.SampleRate;
+	const float AdjustedNewPlaybackFramePosition = FMath::Max((StartSample + FramesToMove) - TransformationChainConfig.StartTime * TransformationChainConfig.SampleRate, 0.0f);
+
+	HandlePlayheadScrub(AdjustedNewPlaybackFramePosition / ActiveDurationInFrames, false);
+	TransportCoordinator->SetProgressRatio(AdjustedNewPlaybackFramePosition / ActiveDurationInFrames);
+}
+
 void FWaveformEditor::RegenerateTransformations()
 {
 	check(SoundWave);
@@ -1597,6 +1629,10 @@ bool FWaveformEditor::BindCommands()
 		Commands.SkipToNextMarker,
 		FExecuteAction::CreateSP(this, &FWaveformEditor::SkipToNextMarker));
 
+	ToolkitCommands->MapAction(
+		Commands.SkipToNextZeroCrossing,
+		FExecuteAction::CreateSP(this, &FWaveformEditor::SkipToNextZeroCrossing));
+
 	return true;
 }
 
diff --git a/Engine/Plugins/Editor/WaveformEditor/Source/WaveformEditor/Private/WaveformEditorCommands.cpp b/Engine/Plugins/Editor/WaveformEditor/Source/WaveformEditor/Private/WaveformEditorCommands.cpp
--- a/Engine/Plugins/Editor/WaveformEditor/Source/WaveformEditor/Private/WaveformEditorCommands.cpp
+++ b/Engine/Plugins/Editor/WaveformEditor/Source/WaveformEditor/Private/WaveformEditorCommands.cpp
@@ -65,6 +65,7 @@ void FWaveformEditorCommands::RegisterCommands()
 	UI_COMMAND(DeleteMarker, "Delete Marker Or Loop Region", "Delete a selected Marker Cue or Loop Region from the sound wave", EUserInterfaceActionType::Button, FInputChord(EKeys::Delete));
 
 	UI_COMMAND(SkipToNextMarker, "Skip To Next Marker", "Skips the playhead to the next marker", EUserInterfaceActionType::Button, FInputChord(EKeys::G));
+	UI_COMMAND(SkipToNextZeroCrossing, "Skip To Next Zero Crossing", "Skips the playhead to the next zero crossing", EUserInterfaceActionType::Button, FInputChord(EKeys::H));
 }
 	
 const FWaveformEditorCommands& FWaveformEditorCommands::Get()
diff --git a/Engine/Plugins/Editor/WaveformEditor/Source/WaveformEditor/Public/WaveformEditor.h b/Engine/Plugins/Editor/WaveformEditor/Source/WaveformEditor/Public/WaveformEditor.h
--- a/Engine/Plugins/Editor/WaveformEditor/Source/WaveformEditor/Public/WaveformEditor.h
+++ b/Engine/Plugins/Editor/WaveformEditor/Source/WaveformEditor/Public/WaveformEditor.h
@@ -134,6 +134,7 @@ private:
 	UE_API bool CanFadeOut();
 	UE_API void CreateMarker(bool bIsLoopRegion);
 	UE_API void DeleteMarker();
+	UE_API void SkipToNextZeroCrossing();
 	UE_API void RegenerateTransformations();
 	UE_API void UpdateTransformations(bool bMarkDirty = true);
 	UE_API void UpdateRenderElements();
diff --git a/Engine/Plugins/Editor/WaveformEditor/Source/WaveformEditor/Public/WaveformEditorCommands.h b/Engine/Plugins/Editor/WaveformEditor/Source/WaveformEditor/Public/WaveformEditorCommands.h
--- a/Engine/Plugins/Editor/WaveformEditor/Source/WaveformEditor/Public/WaveformEditorCommands.h
+++ b/Engine/Plugins/Editor/WaveformEditor/Source/WaveformEditor/Public/WaveformEditorCommands.h
@@ -63,4 +63,5 @@ public:
 	TSharedPtr<FUICommandInfo> DeleteMarker;
 
 	TSharedPtr<FUICommandInfo> SkipToNextMarker;
+	TSharedPtr<FUICommandInfo> SkipToNextZeroCrossing;
 };
diff --git a/Engine/Plugins/Editor/WaveformEditor/Source/WaveformTransformations/Private/WaveformTransformationMarkers.cpp b/Engine/Plugins/Editor/WaveformEditor/Source/WaveformTransformations/Private/WaveformTransformationMarkers.cpp
--- a/Engine/Plugins/Editor/WaveformEditor/Source/WaveformTransformations/Private/WaveformTransformationMarkers.cpp
+++ b/Engine/Plugins/Editor/WaveformEditor/Source/WaveformTransformations/Private/WaveformTransformationMarkers.cpp
@@ -326,7 +326,7 @@ void UWaveformTransformationMarkers::ModifyMarkerLoopRegion(ELoopModificationCon
 		const int64 StartSample = (LoopRegion->FramePosition) * NumChannels;
 		const int64 EndSample = 0;
 
-		const int64 FramesToMove = GetFramesToNextZeroCrossing(StartSample, EndSample, PCMArray, false, NoiseThreshold);
+		const int64 FramesToMove = GetFramesToNextZeroCrossing(StartSample, EndSample, PCMArray, false, NoiseThreshold, NumChannels);
 
 		LoopRegion->FramePosition -= FramesToMove;
 		LoopRegion->FrameLength += FramesToMove;
@@ -337,7 +337,7 @@ void UWaveformTransformationMarkers::ModifyMarkerLoopRegion(ELoopModificationCon
 		const int64 StartSample = (LoopRegion->FramePosition) * NumChannels;
 		const int64 EndSample = (LoopRegion->FramePosition + LoopRegion->FrameLength) * NumChannels;
 
-		const int64 FramesToMove = GetFramesToNextZeroCrossing(StartSample, EndSample, PCMArray, true, NoiseThreshold);
+		const int64 FramesToMove = GetFramesToNextZeroCrossing(StartSample, EndSample, PCMArray, true, NoiseThreshold, NumChannels);
 
 		LoopRegion->FramePosition += FramesToMove;
 		LoopRegion->FrameLength -= FramesToMove;
@@ -348,7 +348,7 @@ void UWaveformTransformationMarkers::ModifyMarkerLoopRegion(ELoopModificationCon
 		const int64 StartSample = (LoopRegion->FramePosition + LoopRegion->FrameLength) * NumChannels;
 		const int64 EndSample = NumAvailableSamples;
 
-		LoopRegion->FrameLength += GetFramesToNextZeroCrossing(StartSample, EndSample, PCMArray, true, NoiseThreshold);
+		LoopRegion->FrameLength += GetFramesToNextZeroCrossing(StartSample, EndSample, PCMArray, true, NoiseThreshold, NumChannels);
 	}
 	break;
 	case ELoopModificationControls::RightHandlePreviousZeroCrossing:
@@ -356,7 +356,7 @@ void UWaveformTransformationMarkers::ModifyMarkerLoopRegion(ELoopModificationCon
 		const int64 StartSample = (LoopRegion->FramePosition + LoopRegion->FrameLength) * NumChannels;
 		const int64 EndSample = LoopRegion->FramePosition * NumChannels;
 
-		LoopRegion->FrameLength -= GetFramesToNextZeroCrossing(StartSample, EndSample, PCMArray, false, NoiseThreshold);
+		LoopRegion->FrameLength -= GetFramesToNextZeroCrossing(StartSample, EndSample, PCMArray, false, NoiseThreshold, NumChannels);
 	}
 	default:
 		break;
@@ -444,9 +444,9 @@ void UWaveformTransformationMarkers::CycleMarkerLoopRegion(ELoopModificationCont
 	OnTransformationChanged.ExecuteIfBound(false);
 }
 
-int64 UWaveformTransformationMarkers::GetFramesToNextZeroCrossing(int64 StartSample, int64 EndSample, const TArray<float>& PCMArray, bool bCheckIncrementing, float NoiseThreshold) const
+int64 UWaveformTransformationMarkers::GetFramesToNextZeroCrossing(int64 StartSample, int64 EndSample, const TArray<float>& PCMArray, bool bCheckIncrementing, float NoiseThreshold, int ChannelCount)
 {
-	check(NumChannels > 0);
+	check(ChannelCount > 0);
 	int64 FramesMoved = 0;
 
 	bool bAllFramesPositiveLastFrame = false;
@@ -455,9 +455,9 @@ int64 UWaveformTransformationMarkers::GetFramesToNextZeroCrossing(int64 StartSam
 
 	if (bCheckIncrementing)
 	{
-		for (int64 i = StartSample; i < EndSample; i += NumChannels)
+		for (int64 i = StartSample; i < EndSample; i += ChannelCount)
 		{
-			if (CheckZeroCrossingFrame(i, PCMArray, bAllFramesPositiveLastFrame, bAllFramesNegativeLastFrame, bAllFramesLowNoiseLastFrame, NoiseThreshold))
+			if (CheckZeroCrossingFrame(i, PCMArray, bAllFramesPositiveLastFrame, bAllFramesNegativeLastFrame, bAllFramesLowNoiseLastFrame, NoiseThreshold, ChannelCount))
 			{
 				return FramesMoved;
 			}
@@ -467,9 +467,9 @@ int64 UWaveformTransformationMarkers::GetFramesToNextZeroCrossing(int64 StartSam
 	}
 	else
 	{
-		for (int64 i = StartSample; i >= EndSample; i -= NumChannels)
+		for (int64 i = StartSample; i >= EndSample; i -= ChannelCount)
 		{
-			if (CheckZeroCrossingFrame(i, PCMArray, bAllFramesPositiveLastFrame, bAllFramesNegativeLastFrame, bAllFramesLowNoiseLastFrame, NoiseThreshold))
+			if (CheckZeroCrossingFrame(i, PCMArray, bAllFramesPositiveLastFrame, bAllFramesNegativeLastFrame, bAllFramesLowNoiseLastFrame, NoiseThreshold, ChannelCount))
 			{
 				return FramesMoved;
 			}
@@ -481,12 +481,12 @@ int64 UWaveformTransformationMarkers::GetFramesToNextZeroCrossing(int64 StartSam
 	return 0;
 }
 
-int64 UWaveformTransformationMarkers::CheckZeroCrossingFrame(int64 CurrentFrame, const TArray<float>& PCMArray, bool& bAllFramesPositiveLastFrame, bool& bAllFramesNegativeLastFrame, bool& bAllFramesLowNoiseLastFrame, float NoiseThreshold) const
+int64 UWaveformTransformationMarkers::CheckZeroCrossingFrame(int64 CurrentFrame, const TArray<float>& PCMArray, bool& bAllFramesPositiveLastFrame, bool& bAllFramesNegativeLastFrame, bool& bAllFramesLowNoiseLastFrame, float NoiseThreshold, int ChannelCount)
 {
-	check(NumChannels > 0);
-	check(CurrentFrame % NumChannels == 0); // Check for frame alignment
+	check(ChannelCount > 0);
+	check(CurrentFrame % ChannelCount == 0); // Check for frame alignment
 
-	if (CurrentFrame + NumChannels > PCMArray.Num())
+	if (CurrentFrame + ChannelCount > PCMArray.Num())
 	{
 		return 0;
 	}
@@ -496,7 +496,7 @@ int64 UWaveformTransformationMarkers::CheckZeroCrossingFrame(int64 CurrentFrame,
 	int32 LowNoiseSampleCount = 0;
 
 	// check if all channels have a zero crossing
-	for (int32 ChannelNum = 0; ChannelNum < NumChannels; ChannelNum++)
+	for (int32 ChannelNum = 0; ChannelNum < ChannelCount; ChannelNum++)
 	{
 		const float FrameValue = PCMArray[CurrentFrame + ChannelNum];
 
@@ -517,9 +517,9 @@ int64 UWaveformTransformationMarkers::CheckZeroCrossingFrame(int64 CurrentFrame,
 		}
 	}
 
-	bool bAllSamplesPositive = PositiveSampleCount == NumChannels;
-	bool bAllSamplesNegative = NegativeSampleCount == NumChannels;
-	bool bAllSamplesLowNoise = LowNoiseSampleCount == NumChannels;
+	bool bAllSamplesPositive = PositiveSampleCount == ChannelCount;
+	bool bAllSamplesNegative = NegativeSampleCount == ChannelCount;
+	bool bAllSamplesLowNoise = LowNoiseSampleCount == ChannelCount;
 
 	// Check if we have crossed the zero crossing on all channels between last frame and this frame
 	// If we are in a low noise frame then we ignore it for zero crossing
diff --git a/Engine/Plugins/Editor/WaveformEditor/Source/WaveformTransformations/Public/WaveformTransformationMarkers.h b/Engine/Plugins/Editor/WaveformEditor/Source/WaveformTransformations/Public/WaveformTransformationMarkers.h
--- a/Engine/Plugins/Editor/WaveformEditor/Source/WaveformTransformations/Public/WaveformTransformationMarkers.h
+++ b/Engine/Plugins/Editor/WaveformEditor/Source/WaveformTransformations/Public/WaveformTransformationMarkers.h
@@ -130,8 +130,8 @@ public:
 	// Set SelectedCue by cycling through cues and loops with toolkit commands bound in the waveform editor
 	UE_API void CycleMarkerLoopRegion(ELoopModificationControls Modification);
 
-	UE_API int64 GetFramesToNextZeroCrossing(int64 StartFrame, int64 EndFrame, const TArray<float>& PCMArray, bool bCheckIncrementing, float NoiseThreshold) const;
-	UE_API int64 CheckZeroCrossingFrame(int64 CurrentFrame, const TArray<float>& PCMArray, bool& bAllFramesPositiveLastFrame, bool& bAllFramesNegativeLastFrame, bool& bAllFramesLowNoiseLastFrame, float NoiseThreshold) const;
+	static UE_API int64 GetFramesToNextZeroCrossing(int64 StartFrame, int64 EndFrame, const TArray<float>& PCMArray, bool bCheckIncrementing, float NoiseThreshold, int ChannelCount);
+	static UE_API int64 CheckZeroCrossingFrame(int64 CurrentFrame, const TArray<float>& PCMArray, bool& bAllFramesPositiveLastFrame, bool& bAllFramesNegativeLastFrame, bool& bAllFramesLowNoiseLastFrame, float NoiseThreshold, int ChannelCount);
 
 	UE_API bool IsLoopRegionActive() const;
 	UE_API void SetLoopPreviewing();
