[Waveform Editor] - Using Remove Marker toggle in Toolbar creates a Marker Transformation

diff --git a/Engine/Plugins/Editor/WaveformEditor/Source/WaveformEditor/Private/WaveformEditor.cpp b/Engine/Plugins/Editor/WaveformEditor/Source/WaveformEditor/Private/WaveformEditor.cpp
--- a/Engine/Plugins/Editor/WaveformEditor/Source/WaveformEditor/Private/WaveformEditor.cpp
+++ b/Engine/Plugins/Editor/WaveformEditor/Source/WaveformEditor/Private/WaveformEditor.cpp
@@ -606,12 +606,18 @@ void FWaveformEditor::DeleteMarker()
 {
 	check(SoundWave);
 	check(GEditor != nullptr);
-	GEditor->BeginTransaction(TEXT("PropertyEditor"), LOCTEXT("SetNumericPropertyTransactionDeleteMarker", "DeleteMarker"), nullptr);
-	SoundWave->Modify();
 
-	TObjectPtr<UWaveformTransformationMarkers> MarkerTransformation = GetOrAddMarkerTransformation();
-	check(MarkerTransformation);
+	TObjectPtr<UWaveformTransformationMarkers> MarkerTransformation = GetMarkerTransformation();
+
+	if (MarkerTransformation == nullptr)
+	{
+		return;
+	}
+
 	check(MarkerTransformation->Markers);
+	
+	GEditor->BeginTransaction(TEXT("PropertyEditor"), LOCTEXT("SetNumericPropertyTransactionDeleteMarker", "DeleteMarker"), nullptr);
+	SoundWave->Modify();
 
 	int32 CueToDelete = MarkerTransformation->Markers->SelectedCue;
 
