[Audio] Add cook diff tags for USoundWave cooking

diff --git a/Engine/Plugins/Editor/WaveformEditor/Source/WaveformTransformations/Private/WaveformTransformationMarkers.cpp b/Engine/Plugins/Editor/WaveformEditor/Source/WaveformTransformations/Private/WaveformTransformationMarkers.cpp
--- a/Engine/Plugins/Editor/WaveformEditor/Source/WaveformTransformations/Private/WaveformTransformationMarkers.cpp
+++ b/Engine/Plugins/Editor/WaveformEditor/Source/WaveformTransformations/Private/WaveformTransformationMarkers.cpp
@@ -2,6 +2,7 @@
 
 #include "WaveformTransformationMarkers.h"
 
+#include "AudioCompressionSettingsUtils.h"
 #include "Editor.h"
 #include "Sound/SoundWave.h"
 #include "WaveformTransformationLog.h"
@@ -699,6 +700,25 @@ void UWaveformTransformationMarkers::GetTransformationInfo(FWaveformTransformati
 	InOutTransformationInfo.AllCuePointsRelativeToStartTime.Append(GetAdjustedCuePoints(Markers->CuesAndLoops, StartFrameOffset, bRemoveTrimmed));
 }
 
+FString UWaveformTransformationMarkers::GetTransformationHash() const
+{
+	check(Markers);
+
+	using FPCU = FPlatformCompressionUtilities;
+	FString Hash;
+	
+	for (const FSoundWaveCuePoint& CuePoint : Markers->CuesAndLoops)
+	{
+		FPCU::AppendHash(Hash, TEXT("CPID"), CuePoint.CuePointID);
+		FPCU::AppendHash(Hash, TEXT("CPFP"), CuePoint.FramePosition);
+		FPCU::AppendHash(Hash, TEXT("CPFL"), CuePoint.FrameLength);
+		FPCU::AppendHash(Hash, TEXT("CPILR"), CuePoint.IsLoopRegion());
+		FPCU::AppendHash(Hash, TEXT("CPL"), CuePoint.Label);
+	}
+
+	return Hash;
+}
+
 FSoundWaveCuePoint* UWaveformTransformationMarkers::GetSelectedMarker() const
 {
 	check(Markers);
diff --git a/Engine/Plugins/Editor/WaveformEditor/Source/WaveformTransformations/Private/WaveformTransformationNormalize.cpp b/Engine/Plugins/Editor/WaveformEditor/Source/WaveformTransformations/Private/WaveformTransformationNormalize.cpp
--- a/Engine/Plugins/Editor/WaveformEditor/Source/WaveformTransformations/Private/WaveformTransformationNormalize.cpp
+++ b/Engine/Plugins/Editor/WaveformEditor/Source/WaveformTransformations/Private/WaveformTransformationNormalize.cpp
@@ -2,6 +2,7 @@
 
 #include "WaveformTransformationNormalize.h"
 
+#include "AudioCompressionSettingsUtils.h"
 #include "DSP/Dsp.h"
 #include "DSP/FloatArrayMath.h"
 #include "WaveformAudioAnalysisFunctions.h"
@@ -61,3 +62,17 @@ Audio::FTransformationPtr UWaveformTransformationNormalize::CreateTransformation
 {
 	return MakeUnique<FWaveTransformationNormalize>(Target, MaxGain, Mode);
 }
+
+#if WITH_EDITOR
+FString UWaveformTransformationNormalize::GetTransformationHash() const
+{
+	using FPCU = FPlatformCompressionUtilities;
+	FString Hash;
+
+	FPCU::AppendHash(Hash, TEXT("NMG"),MaxGain);
+	FPCU::AppendHash(Hash, TEXT("NT"), Target);
+	FPCU::AppendHash(Hash, TEXT("NM"), static_cast<uint8>(Mode));
+
+	return Hash;
+}
+#endif //WITH_EDITOR
\ No newline at end of file
diff --git a/Engine/Plugins/Editor/WaveformEditor/Source/WaveformTransformations/Private/WaveformTransformationTrimFade.cpp b/Engine/Plugins/Editor/WaveformEditor/Source/WaveformTransformations/Private/WaveformTransformationTrimFade.cpp
--- a/Engine/Plugins/Editor/WaveformEditor/Source/WaveformTransformations/Private/WaveformTransformationTrimFade.cpp
+++ b/Engine/Plugins/Editor/WaveformEditor/Source/WaveformTransformations/Private/WaveformTransformationTrimFade.cpp
@@ -2,18 +2,20 @@
 
 #include "WaveformTransformationTrimFade.h"
 
+#include "AudioCompressionSettingsUtils.h"
 #include "Editor.h"
 #include "Sound/SoundWave.h"
 #include "WaveformTransformationLog.h"
 
 #include UE_INLINE_GENERATED_CPP_BY_NAME(WaveformTransformationTrimFade)
 
-namespace
+namespace WaveformTrasnformationTrimFadePrivate
 {
 	// Makes the two functions for SCurve meet at (0.5, 0.5)
 	const double XMultiplier = FMath::Sqrt(10.0) / 5.0;
 
-	UE_DEPRECATED(5.7, "Individual fade properties have been deprecated. Use FadeFunctions::FadeIn::GetFadeInCurveValue instead.")
+	constexpr float Linear = 1.f;
+
 	void SetFadeCurve(TObjectPtr<UFadeFunction> FadeFunction, float FadeCurve, float Sharpness)
 	{
 		if (FadeFunction->IsA<UFadeCurveFunction>())
@@ -233,14 +235,12 @@ void FWaveTransformationTrimFade::ProcessAudio(Audio::FWaveformTransformationWav
 
 UWaveformTransformationTrimFade::UWaveformTransformationTrimFade(const FObjectInitializer& ObjectInitializer) : Super(ObjectInitializer)
 {
-	const float Linear = 1.f;
-
 	if (FadeFunctions.FadeIn == nullptr)
 	{
 		PRAGMA_DISABLE_DEPRECATION_WARNINGS
 			if (StartFadeTime != 0.f)
 			{
-				if (StartFadeCurve > Linear)
+				if (StartFadeCurve > WaveformTrasnformationTrimFadePrivate::Linear)
 				{
 					FadeFunctions.FadeIn = CreateDefaultSubobject<UFadeCurveFunctionExponential>(TEXT("FadeInFunction"));
 				}
@@ -248,17 +248,17 @@ UWaveformTransformationTrimFade::UWaveformTransformationTrimFade(const FObjectIn
 				{
 					FadeFunctions.FadeIn = CreateDefaultSubobject<UFadeCurveFunctionSigmoid>(TEXT("FadeInFunction"));
 				}
-				else if (StartFadeCurve < Linear)
+				else if (StartFadeCurve < WaveformTrasnformationTrimFadePrivate::Linear)
 				{
 					FadeFunctions.FadeIn = CreateDefaultSubobject<UFadeCurveFunctionLogarithmic>(TEXT("FadeInFunction"));
 				}
-				else if (StartFadeCurve == Linear)
+				else if (StartFadeCurve == WaveformTrasnformationTrimFadePrivate::Linear)
 				{
 					FadeFunctions.FadeIn = CreateDefaultSubobject<UFadeFunctionLinear>(TEXT("FadeInFunction"));
 				}
 
 				FadeFunctions.FadeIn->Duration = StartFadeTime;
-				SetFadeCurve(FadeFunctions.FadeIn, StartFadeCurve, StartSCurveSharpness);
+				WaveformTrasnformationTrimFadePrivate::SetFadeCurve(FadeFunctions.FadeIn, StartFadeCurve, StartSCurveSharpness);
 
 				StartFadeTime = 0.f;
 				PRAGMA_ENABLE_DEPRECATION_WARNINGS
@@ -276,7 +276,7 @@ UWaveformTransformationTrimFade::UWaveformTransformationTrimFade(const FObjectIn
 		PRAGMA_DISABLE_DEPRECATION_WARNINGS
 			if (EndFadeTime != 0.f)
 			{
-				if (EndFadeCurve > Linear)
+				if (EndFadeCurve > WaveformTrasnformationTrimFadePrivate::Linear)
 				{
 					FadeFunctions.FadeOut = CreateDefaultSubobject<UFadeCurveFunctionExponential>(TEXT("FadeOutFunction"));
 				}
@@ -284,17 +284,17 @@ UWaveformTransformationTrimFade::UWaveformTransformationTrimFade(const FObjectIn
 				{
 					FadeFunctions.FadeOut = CreateDefaultSubobject<UFadeCurveFunctionSigmoid>(TEXT("FadeOutFunction"));
 				}
-				else if (EndFadeCurve < Linear)
+				else if (EndFadeCurve < WaveformTrasnformationTrimFadePrivate::Linear)
 				{
 					FadeFunctions.FadeOut = CreateDefaultSubobject<UFadeCurveFunctionLogarithmic>(TEXT("FadeOutFunction"));
 				}
-				else if (EndFadeCurve == Linear)
+				else if (EndFadeCurve == WaveformTrasnformationTrimFadePrivate::Linear)
 				{
 					FadeFunctions.FadeOut = CreateDefaultSubobject<UFadeFunctionLinear>(TEXT("FadeOutFunction"));
 				}
 
 				FadeFunctions.FadeOut->Duration = EndFadeTime;
-				SetFadeCurve(FadeFunctions.FadeOut, EndFadeCurve, EndSCurveSharpness);
+				WaveformTrasnformationTrimFadePrivate::SetFadeCurve(FadeFunctions.FadeOut, EndFadeCurve, EndSCurveSharpness);
 
 				EndFadeTime = 0.f;
 				PRAGMA_ENABLE_DEPRECATION_WARNINGS
@@ -324,19 +324,58 @@ void UWaveformTransformationTrimFade::PostEditChangeProperty(FPropertyChangedEve
 		FadeFunctions.FadeOut->Duration = FMath::Min(FadeFunctions.FadeOut->Duration, MaxDuration);
 	}
 }
+
+FString UWaveformTransformationTrimFade::GetTransformationHash() const
+{
+	using FPCU = FPlatformCompressionUtilities;
+	FString Hash;
+
+	float FadeInDuration = 0.f;
+	float FadeInCurve = WaveformTrasnformationTrimFadePrivate::Linear;
+
+	float FadeOutDuration = 0.f;
+	float FadeOutCurve = WaveformTrasnformationTrimFadePrivate::Linear;
+
+	if (FadeFunctions.FadeIn)
+	{
+		FadeInDuration = FadeFunctions.FadeIn->Duration;
+		if (FadeFunctions.FadeIn.IsA<UFadeCurveFunction>())
+		{
+			TObjectPtr<UFadeCurveFunction> CurveFunction = CastChecked<UFadeCurveFunction>(FadeFunctions.FadeIn);
+			FadeInCurve = CurveFunction->GetFadeCurve();
+		}
+	}
+
+	if (FadeFunctions.FadeOut)
+	{
+		FadeOutDuration = FadeFunctions.FadeOut->Duration;
+		if (FadeFunctions.FadeOut.IsA<UFadeCurveFunction>())
+		{
+			TObjectPtr<UFadeCurveFunction> CurveFunction = CastChecked<UFadeCurveFunction>(FadeFunctions.FadeOut);
+			FadeOutCurve = CurveFunction->GetFadeCurve();
+		}
+	}
+
+	FPCU::AppendHash(Hash, TEXT("TFST"), StartTime);
+	FPCU::AppendHash(Hash, TEXT("TFET"), EndTime);
+	FPCU::AppendHash(Hash, TEXT("TFIC"), FadeInCurve);
+	FPCU::AppendHash(Hash, TEXT("TFID"), FadeInDuration);
+	FPCU::AppendHash(Hash, TEXT("TFOC"), FadeOutCurve);
+	FPCU::AppendHash(Hash, TEXT("TFOD"), FadeOutDuration);
+	
+	return Hash;
+}
 #endif // WITH_EDITOR
 
 void UWaveformTransformationTrimFade::PostLoad()
 {
 	Super::PostLoad();
 
-	const float Linear = 1.f;
-
 	PRAGMA_DISABLE_DEPRECATION_WARNINGS
 	// Overwrite exiting fade in function data with deprecated curve properties to migrate the data to the new properties
 	if (StartFadeTime != 0.f)
 	{
-		if (StartFadeCurve > Linear)
+		if (StartFadeCurve > WaveformTrasnformationTrimFadePrivate::Linear)
 		{
 			FadeFunctions.FadeIn = NewObject<UFadeCurveFunctionExponential>(this, UFadeCurveFunctionExponential::StaticClass(), NAME_None, RF_Transactional);
 		}
@@ -344,17 +383,17 @@ void UWaveformTransformationTrimFade::PostLoad()
 		{
 			FadeFunctions.FadeIn = NewObject<UFadeCurveFunctionSigmoid>(this, UFadeCurveFunctionSigmoid::StaticClass(), NAME_None, RF_Transactional);
 		}
-		else if (StartFadeCurve < Linear)
+		else if (StartFadeCurve < WaveformTrasnformationTrimFadePrivate::Linear)
 		{
 			FadeFunctions.FadeIn = NewObject<UFadeCurveFunctionLogarithmic>(this, UFadeCurveFunctionLogarithmic::StaticClass(), NAME_None, RF_Transactional);
 		}
-		else if (StartFadeCurve == Linear)
+		else if (StartFadeCurve == WaveformTrasnformationTrimFadePrivate::Linear)
 		{
 			FadeFunctions.FadeIn = NewObject<UFadeFunctionLinear>(this, UFadeFunctionLinear::StaticClass(), NAME_None, RF_Transactional);
 		}
 
 		FadeFunctions.FadeIn->Duration = StartFadeTime;
-		SetFadeCurve(FadeFunctions.FadeIn, StartFadeCurve, StartSCurveSharpness);
+		WaveformTrasnformationTrimFadePrivate::SetFadeCurve(FadeFunctions.FadeIn, StartFadeCurve, StartSCurveSharpness);
 
 		StartFadeTime = 0.f;
 	}
@@ -369,7 +408,7 @@ void UWaveformTransformationTrimFade::PostLoad()
 	// Overwrite exiting fade out function data with deprecated curve properties to migrate the data to the new properties
 	if (EndFadeTime != 0.f)
 	{
-		if (EndFadeCurve > Linear)
+		if (EndFadeCurve > WaveformTrasnformationTrimFadePrivate::Linear)
 		{
 			FadeFunctions.FadeOut = NewObject<UFadeCurveFunctionExponential>(this, UFadeCurveFunctionExponential::StaticClass(), NAME_None, RF_Transactional);
 		}
@@ -377,17 +416,17 @@ void UWaveformTransformationTrimFade::PostLoad()
 		{
 			FadeFunctions.FadeOut = NewObject<UFadeCurveFunctionSigmoid>(this, UFadeCurveFunctionSigmoid::StaticClass(), NAME_None, RF_Transactional);
 		}
-		else if (EndFadeCurve < Linear)
+		else if (EndFadeCurve < WaveformTrasnformationTrimFadePrivate::Linear)
 		{
 			FadeFunctions.FadeOut = NewObject<UFadeCurveFunctionLogarithmic>(this, UFadeCurveFunctionLogarithmic::StaticClass(), NAME_None, RF_Transactional);
 		}
-		else if (EndFadeCurve == Linear)
+		else if (EndFadeCurve == WaveformTrasnformationTrimFadePrivate::Linear)
 		{
 			FadeFunctions.FadeOut = NewObject<UFadeFunctionLinear>(this, UFadeFunctionLinear::StaticClass(), NAME_None, RF_Transactional);
 		}
 
 		FadeFunctions.FadeOut->Duration = EndFadeTime;
-		SetFadeCurve(FadeFunctions.FadeOut, EndFadeCurve, EndSCurveSharpness);
+		WaveformTrasnformationTrimFadePrivate::SetFadeCurve(FadeFunctions.FadeOut, EndFadeCurve, EndSCurveSharpness);
 
 		EndFadeTime = 0.f;
 	}
@@ -425,11 +464,11 @@ const double UWaveformTransformationTrimFade::GetFadeInCurveValue(const float St
 
 		if (FadeFraction <= 0.5)
 		{
-			CurveValue = -(FMath::Tanh(K * PI * FMath::LogX(10, XMultiplier * X) + ((K * PI) / 2))) / 2 + 0.5;
+			CurveValue = -(FMath::Tanh(K * PI * FMath::LogX(10, WaveformTrasnformationTrimFadePrivate::XMultiplier * X) + ((K * PI) / 2))) / 2 + 0.5;
 		}
 		else
 		{
-			CurveValue = (FMath::Tanh(K * PI * FMath::LogX(10, XMultiplier * (-X + 1)) + ((K * PI) / 2))) / 2 + 0.5;
+			CurveValue = (FMath::Tanh(K * PI * FMath::LogX(10, WaveformTrasnformationTrimFadePrivate::XMultiplier * (-X + 1)) + ((K * PI) / 2))) / 2 + 0.5;
 		}
 	}
 	else
@@ -452,11 +491,11 @@ const double UWaveformTransformationTrimFade::GetFadeOutCurveValue(const float E
 
 		if (FadeFraction <= 0.5)
 		{
-			CurveValue = -(FMath::Tanh(K * PI * FMath::LogX(10, XMultiplier * X) + ((K * PI) / 2))) / 2 + 0.5;
+			CurveValue = -(FMath::Tanh(K * PI * FMath::LogX(10, WaveformTrasnformationTrimFadePrivate::XMultiplier * X) + ((K * PI) / 2))) / 2 + 0.5;
 		}
 		else
 		{
-			CurveValue = (FMath::Tanh(K * PI * FMath::LogX(10, XMultiplier * (-X + 1)) + ((K * PI) / 2))) / 2 + 0.5;
+			CurveValue = (FMath::Tanh(K * PI * FMath::LogX(10, WaveformTrasnformationTrimFadePrivate::XMultiplier * (-X + 1)) + ((K * PI) / 2))) / 2 + 0.5;
 		}
 	}
 	else
@@ -533,7 +572,7 @@ const double UFadeCurveFunctionSigmoid::GetFadeInCurveValue(const double FadeFra
 	// LogX will produce NaN if FadeFraction is <= 0 or is 1
 	check(FadeFraction <= 1.f);
 	check(FadeFraction >= 0.f);
-	check(XMultiplier > 0.0);
+	check(WaveformTrasnformationTrimFadePrivate::XMultiplier > 0.0);
 
 	if (FadeFraction == 0.f || FadeFraction == 1.f)
 	{
@@ -547,11 +586,11 @@ const double UFadeCurveFunctionSigmoid::GetFadeInCurveValue(const double FadeFra
 
 	if (FadeFraction <= 0.5)
 	{
-		CurveValue = -(FMath::Tanh(k * PI * FMath::LogX(10, XMultiplier * x) + ((k * PI) / 2))) / 2 + 0.5;
+		CurveValue = -(FMath::Tanh(k * PI * FMath::LogX(10, WaveformTrasnformationTrimFadePrivate::XMultiplier * x) + ((k * PI) / 2))) / 2 + 0.5;
 	}
 	else
 	{
-		CurveValue = (FMath::Tanh(k * PI * FMath::LogX(10, XMultiplier * (-x + 1)) + ((k * PI) / 2))) / 2 + 0.5;
+		CurveValue = (FMath::Tanh(k * PI * FMath::LogX(10, WaveformTrasnformationTrimFadePrivate::XMultiplier * (-x + 1)) + ((k * PI) / 2))) / 2 + 0.5;
 	}
 
 	return CurveValue;
@@ -562,7 +601,7 @@ const double UFadeCurveFunctionSigmoid::GetFadeOutCurveValue(const double FadeFr
 	// LogX will produce NaN if FadeFraction is <= 0 or is 1
 	check(FadeFraction <= 1.f);
 	check(FadeFraction >= 0.f);
-	check(XMultiplier > 0.0);
+	check(WaveformTrasnformationTrimFadePrivate::XMultiplier > 0.0);
 
 	if (FadeFraction == 0.f || FadeFraction == 1.f)
 	{
@@ -576,11 +615,11 @@ const double UFadeCurveFunctionSigmoid::GetFadeOutCurveValue(const double FadeFr
 
 	if (FadeFraction <= 0.5)
 	{
-		CurveValue = -(FMath::Tanh(k * PI * FMath::LogX(10, XMultiplier * x) + ((k * PI) / 2))) / 2 + 0.5;
+		CurveValue = -(FMath::Tanh(k * PI * FMath::LogX(10, WaveformTrasnformationTrimFadePrivate::XMultiplier * x) + ((k * PI) / 2))) / 2 + 0.5;
 	}
 	else
 	{
-		CurveValue = (FMath::Tanh(k * PI * FMath::LogX(10, XMultiplier * (-x + 1)) + ((k * PI) / 2))) / 2 + 0.5;
+		CurveValue = (FMath::Tanh(k * PI * FMath::LogX(10, WaveformTrasnformationTrimFadePrivate::XMultiplier * (-x + 1)) + ((k * PI) / 2))) / 2 + 0.5;
 	}
 
 	return CurveValue;
diff --git a/Engine/Plugins/Editor/WaveformEditor/Source/WaveformTransformations/Public/WaveformTransformationEffectChain.h b/Engine/Plugins/Editor/WaveformEditor/Source/WaveformTransformations/Public/WaveformTransformationEffectChain.h
--- a/Engine/Plugins/Editor/WaveformEditor/Source/WaveformTransformations/Public/WaveformTransformationEffectChain.h
+++ b/Engine/Plugins/Editor/WaveformEditor/Source/WaveformTransformations/Public/WaveformTransformationEffectChain.h
@@ -38,6 +38,7 @@ class UWaveformTransformationEffectChain : public UWaveformTransformationBase
 public:
 	
 	UE_API virtual Audio::FTransformationPtr CreateTransformation() const override;
+	virtual FString GetTransformationHash() const override { return ""; }
 };
 
 #undef UE_API
diff --git a/Engine/Plugins/Editor/WaveformEditor/Source/WaveformTransformations/Public/WaveformTransformationMarkers.h b/Engine/Plugins/Editor/WaveformEditor/Source/WaveformTransformations/Public/WaveformTransformationMarkers.h
--- a/Engine/Plugins/Editor/WaveformEditor/Source/WaveformTransformations/Public/WaveformTransformationMarkers.h
+++ b/Engine/Plugins/Editor/WaveformEditor/Source/WaveformTransformations/Public/WaveformTransformationMarkers.h
@@ -90,7 +90,7 @@ private:
 };
 
 UCLASS(MinimalAPI)
-class UWaveformTransformationMarkers : public UWaveformTransformationBase
+class UWaveformTransformationMarkers final : public UWaveformTransformationBase
 {
 	GENERATED_BODY()
 public:
@@ -142,8 +142,9 @@ public:
 	UE_API bool AdjustLoopPreviewIfNotAligned();
 
 #if WITH_EDITOR
-	UE_API void OverwriteSoundWaveData(USoundWave& InOutSoundWave) override;
-	UE_API void GetTransformationInfo(FWaveformTransformationInfo& InOutTransformationInfo) const override;
+	UE_API virtual void OverwriteSoundWaveData(USoundWave& InOutSoundWave) override;
+	UE_API virtual void GetTransformationInfo(FWaveformTransformationInfo& InOutTransformationInfo) const override;
+	UE_API virtual FString GetTransformationHash() const override;
 #endif //WITH_EDITOR
 
 	UE_API FSoundWaveCuePoint* GetSelectedMarker() const;
diff --git a/Engine/Plugins/Editor/WaveformEditor/Source/WaveformTransformations/Public/WaveformTransformationNormalize.h b/Engine/Plugins/Editor/WaveformEditor/Source/WaveformTransformations/Public/WaveformTransformationNormalize.h
--- a/Engine/Plugins/Editor/WaveformEditor/Source/WaveformTransformations/Public/WaveformTransformationNormalize.h
+++ b/Engine/Plugins/Editor/WaveformEditor/Source/WaveformTransformations/Public/WaveformTransformationNormalize.h
@@ -31,7 +31,7 @@ private:
 };
 
 UCLASS(MinimalAPI)
-class UWaveformTransformationNormalize : public UWaveformTransformationBase
+class UWaveformTransformationNormalize final : public UWaveformTransformationBase
 {
 	GENERATED_BODY()
 
@@ -50,6 +50,11 @@ class UWaveformTransformationNormalize : public UWaveformTransformationBase
 public:
 	UE_API virtual Audio::FTransformationPtr CreateTransformation() const override;
 	virtual constexpr Audio::ETransformationPriority GetTransformationPriority() const { return Audio::ETransformationPriority::Low; }
+
+#if WITH_EDITOR
+	UE_API virtual FString GetTransformationHash() const override;
+#endif // WITH_EDITOR
+
 };
 
 #undef UE_API
diff --git a/Engine/Plugins/Editor/WaveformEditor/Source/WaveformTransformations/Public/WaveformTransformationTrimFade.h b/Engine/Plugins/Editor/WaveformEditor/Source/WaveformTransformations/Public/WaveformTransformationTrimFade.h
--- a/Engine/Plugins/Editor/WaveformEditor/Source/WaveformTransformations/Public/WaveformTransformationTrimFade.h
+++ b/Engine/Plugins/Editor/WaveformEditor/Source/WaveformTransformations/Public/WaveformTransformationTrimFade.h
@@ -143,7 +143,7 @@ private:
 };
 
 UCLASS(MinimalAPI)
-class UWaveformTransformationTrimFade : public UWaveformTransformationBase
+class UWaveformTransformationTrimFade final : public UWaveformTransformationBase
 {
 	GENERATED_BODY()
 public:
@@ -152,6 +152,8 @@ public:
 
 #if WITH_EDITOR
 	UE_API virtual void PostEditChangeProperty(FPropertyChangedEvent& PropertyChangedEvent) override;
+
+	UE_API virtual FString GetTransformationHash() const override;
 #endif // WITH_EDITOR
 
 	UE_API virtual void PostLoad() override;
diff --git a/Engine/Source/Runtime/AudioExtensions/Public/IWaveformTransformation.h b/Engine/Source/Runtime/AudioExtensions/Public/IWaveformTransformation.h
--- a/Engine/Source/Runtime/AudioExtensions/Public/IWaveformTransformation.h
+++ b/Engine/Source/Runtime/AudioExtensions/Public/IWaveformTransformation.h
@@ -125,7 +125,7 @@ struct FWaveTransformUObjectConfiguration
 };
 
 #if WITH_EDITORONLY_DATA
-// Information to be retrieved from each transformation (Add members as needed when new transformation types are added)
+// Information to be retrieved from each transformation (Add members as needed)
 struct FWaveformTransformationInfo
 {
 	TArray<FSoundWaveCuePoint> AllCuePoints; // Cue Points and Loop Regions
@@ -166,6 +166,7 @@ public:
 #if WITH_EDITOR
 	virtual void OverwriteSoundWaveData(USoundWave& InOutSoundWave) {};
 	virtual void GetTransformationInfo(FWaveformTransformationInfo& InOutTransformationInfo) const {};
+	virtual FString GetTransformationHash() const PURE_VIRTUAL(UWaveformTransformationBase::GetTransformationHash, return FString(););
 
 	AUDIOEXTENSIONS_API virtual void PostEditUndo() override;
 	AUDIOEXTENSIONS_API virtual void PostEditChangeProperty(FPropertyChangedEvent& PropertyChangedEvent) override;
diff --git a/Engine/Source/Runtime/Engine/Private/AudioDerivedData.cpp b/Engine/Source/Runtime/Engine/Private/AudioDerivedData.cpp
--- a/Engine/Source/Runtime/Engine/Private/AudioDerivedData.cpp
+++ b/Engine/Source/Runtime/Engine/Private/AudioDerivedData.cpp
@@ -229,6 +229,41 @@ static FString GetSoundWaveHash(const USoundWave& InWave, const ITargetPlatform*
 	return SoundWaveHash;
 }
 
+static FSHAHash GetSoundWaveRawDataHash(const USoundWave& InWave)
+{
+	FSHAHash Out{};
+	const USoundWave::FEditorAudioBulkData& RawData = InWave.RawData;
+
+	const FSharedBuffer& Payload = RawData.GetPayload().Consume();
+	const uint64 NumBytes = Payload.GetSize();
+
+	if (NumBytes == 0)
+	{
+		return Out;
+	}
+
+	FSHA1::HashBuffer(Payload.GetData(), NumBytes, Out.Hash);
+	return Out;
+}
+
+static FString GetWaveformTransformationHash(const USoundWave& InWave)
+{
+	// Hash the SoundWave waveform transformations 
+	FString WaveformTransformationHash;
+
+	const TArray<TObjectPtr<UWaveformTransformationBase>>& Transformations = InWave.Transformations;
+	
+	for (const TObjectPtr<UWaveformTransformationBase>& Transformation : Transformations)
+	{
+		if (Transformation)
+		{
+			WaveformTransformationHash += Transformation->GetTransformationHash();
+		}
+	}
+
+	return WaveformTransformationHash;
+}
+
 /**
  * Computes the derived data key suffix for a SoundWave's Streamed Audio.
  * @param SoundWave - The SoundWave for which to compute the derived data key.
@@ -1507,6 +1542,18 @@ void FStreamedAudioPlatformData::Serialize(FArchive& Ar, USoundWave* Owner)
 	{
 		FCookTagList* CookTags = Ar.GetCookContext()->GetCookTagList();
 		CookTags->Add(Owner, "StreamingFormat", LexToString(AudioFormat));
+
+		const ITargetPlatform* CookingTarget = Ar.CookingTarget();
+		check(CookingTarget);
+
+		const FName Format = CookingTarget->GetWaveFormat(Owner);
+
+		/* note: must have "Diff_" prefix for the diff tool to pick up the tag as important for diffing */
+		CookTags->Add(Owner, "Diff_10_RawDataHash", GetSoundWaveRawDataHash(*Owner).ToString());
+		CookTags->Add(Owner, "Diff_20_WaveTransformHash", FString(GetWaveformTransformationHash(*Owner)));
+		CookTags->Add(Owner, "Diff_30_FormatHash", Format.ToString());
+		CookTags->Add(Owner, "Diff_40_SoundWaveHash", FString(GetSoundWaveHash(*Owner, CookingTarget)));
+		CookTags->Add(Owner, "Diff_50_CompressedDataGuid", Owner->CompressedDataGuid.ToString());
 	}
 #endif
 
