[Waveform Editor] Undo changes in Waveform Transformations Display does not update visuals in Waveform Editor

#Synced-CL 48504062

diff --git a/Engine/Plugins/Editor/WaveformEditor/Source/WaveformTransformationsWidgets/Private/WaveformTransformationMarkersRenderer.cpp b/Engine/Plugins/Editor/WaveformEditor/Source/WaveformTransformationsWidgets/Private/WaveformTransformationMarkersRenderer.cpp
--- a/Engine/Plugins/Editor/WaveformEditor/Source/WaveformTransformationsWidgets/Private/WaveformTransformationMarkersRenderer.cpp
+++ b/Engine/Plugins/Editor/WaveformEditor/Source/WaveformTransformationsWidgets/Private/WaveformTransformationMarkersRenderer.cpp
@@ -21,7 +21,17 @@ FWaveformTransformationMarkerRenderer::FWaveformTransformationMarkerRenderer()
 	check(Settings);
 	OnSettingChangedDelegateHandle = Settings->OnSettingChanged().AddLambda([this](const FName& PropertyName, const UWaveformTransformationsWidgetsSettings* Settings)
 		{
-			check(Settings)
+			check(Settings);
+
+			if (PropertyName.IsNone())
+			{
+				MarkerColor = Settings->MarkerColor;
+				SetLoopColors(Settings->LoopColors);
+				LabelTextColor = Settings->LabelTextColor;
+				LabelFontSize = Settings->LabelFontSize;
+
+				return;
+			}
 
 			if (PropertyName == GET_MEMBER_NAME_CHECKED(UWaveformTransformationsWidgetsSettings, MarkerColor))
 			{
diff --git a/Engine/Plugins/Editor/WaveformEditor/Source/WaveformTransformationsWidgets/Private/WaveformTransformationsWidgetsSettings.cpp b/Engine/Plugins/Editor/WaveformEditor/Source/WaveformTransformationsWidgets/Private/WaveformTransformationsWidgetsSettings.cpp
--- a/Engine/Plugins/Editor/WaveformEditor/Source/WaveformTransformationsWidgets/Private/WaveformTransformationsWidgetsSettings.cpp
+++ b/Engine/Plugins/Editor/WaveformEditor/Source/WaveformTransformationsWidgets/Private/WaveformTransformationsWidgetsSettings.cpp
@@ -39,6 +39,13 @@ void UWaveformTransformationsWidgetsSettings::PostEditChangeProperty(FPropertyCh
 	}
 }
 
+void UWaveformTransformationsWidgetsSettings::PostEditUndo()
+{
+	Super::PostEditUndo();
+
+	SettingsChangedDelegate.Broadcast(NAME_None, this);
+}
+
 FOnWaveformTransformationsWidgetsSettingsChanged& UWaveformTransformationsWidgetsSettings::OnSettingChanged()
 {
 	return SettingsChangedDelegate;
diff --git a/Engine/Plugins/Editor/WaveformEditor/Source/WaveformTransformationsWidgets/Private/WaveformTransformationsWidgetsSettings.h b/Engine/Plugins/Editor/WaveformEditor/Source/WaveformTransformationsWidgets/Private/WaveformTransformationsWidgetsSettings.h
--- a/Engine/Plugins/Editor/WaveformEditor/Source/WaveformTransformationsWidgets/Private/WaveformTransformationsWidgetsSettings.h
+++ b/Engine/Plugins/Editor/WaveformEditor/Source/WaveformTransformationsWidgets/Private/WaveformTransformationsWidgetsSettings.h
@@ -35,6 +35,7 @@ public:
 	virtual FText GetSectionText() const override;
 
 	virtual void PostEditChangeProperty(FPropertyChangedEvent& PropertyChangedEvent) override;
+	virtual void PostEditUndo() override;
 	//~ End UDeveloperSettings interface
 
 	static FOnWaveformTransformationsWidgetsSettingsChanged& OnSettingChanged();
