[Waveform Editor] - Fade-In handle curve flickers when adjusting Fade duration
#Synced-CL 48802398

diff --git a/Engine/Plugins/Editor/WaveformEditor/Source/WaveformTransformationsWidgets/Private/WaveformTransformationTrimFadeRenderer.cpp b/Engine/Plugins/Editor/WaveformEditor/Source/WaveformTransformationsWidgets/Private/WaveformTransformationTrimFadeRenderer.cpp
--- a/Engine/Plugins/Editor/WaveformEditor/Source/WaveformTransformationsWidgets/Private/WaveformTransformationTrimFadeRenderer.cpp
+++ b/Engine/Plugins/Editor/WaveformEditor/Source/WaveformTransformationsWidgets/Private/WaveformTransformationTrimFadeRenderer.cpp
@@ -178,17 +178,28 @@ void FWaveformTransformationTrimFadeRenderer::GenerateFadeCurves(const FGeometry
 		FadeInEndX = FMath::RoundToInt32(FMath::Clamp(StartTimeHandleX + FadeInPixelLength, StartTimeHandleX, EndTimeHandleX));
 
 		const uint32 DisplayedFadeInPixelLength = FadeInEndX - FadeInStartX;
+		
 		FadeInCurvePoints.SetNumUninitialized(DisplayedFadeInPixelLength);
 
-		for (uint32 Pixel = 0; Pixel < DisplayedFadeInPixelLength; ++Pixel)
+		if (DisplayedFadeInPixelLength > 0)
 		{
-			const double FadeFraction = static_cast<float>(Pixel) / FadeInPixelLength;
-			const double CurveFunction = StrongTrimFade->FadeFunctions.FadeIn->GetFadeInCurveValue(FadeFraction);
+			for (uint32 Pixel = 0; Pixel < DisplayedFadeInPixelLength; ++Pixel)
+			{
+				const double FadeFraction = static_cast<float>(Pixel) / FadeInPixelLength;
+				const double CurveFunction = StrongTrimFade->FadeFunctions.FadeIn->GetFadeInCurveValue(FadeFraction);
+
+				const bool bNeedsFlickerPrevention = (Pixel == DisplayedFadeInPixelLength - 1) && DisplayedFadeInPixelLength != 1;
 
-			const double CurveValue = Pixel != FadeInPixelLength - 1 ? 1.f - CurveFunction : 0.f;
+				// Prevent fade line from flickering if the math doesn't work out to the last Y value being 0 as the position is changing
+				// This is due to pixel rounding errors. The slight distortion caused by this correction isn't noticeable
+				const double CurveValue = bNeedsFlickerPrevention ? 0.f : 1.f - CurveFunction;
 
-			const uint32 XCoordinate = Pixel + FadeInStartX;
-			FadeInCurvePoints[Pixel] = FVector2D(XCoordinate, CurveValue * AllottedGeometry.Size.Y);
+				const uint32 XCoordinate = Pixel + FadeInStartX;
+				FadeInCurvePoints[Pixel] = FVector2D(XCoordinate, CurveValue * AllottedGeometry.Size.Y);
+			}
+
+			check(FadeInCurvePoints.Num() > 0);
+			check(FadeInCurvePoints[0].Y == AllottedGeometry.Size.Y);
 		}
 	}
 	else
@@ -208,18 +219,28 @@ void FWaveformTransformationTrimFadeRenderer::GenerateFadeCurves(const FGeometry
 		FadeOutStartX = FMath::RoundToInt32(FMath::Clamp(EndTimeHandleX - FadeOutPixelLength, StartTimeHandleX, EndTimeHandleX));
 
 		const uint32 DisplayedFadeOutPixelLength = FadeOutEndX - FadeOutStartX;
-		FadeOutCurvePoints.SetNumUninitialized(DisplayedFadeOutPixelLength);
-		const uint32 FadeOutPixelOffset = FMath::RoundToInt32(FMath::Max(FadeOutPixelLength - static_cast<float>(DisplayedFadeOutPixelLength), 0.f));
 
-		for (uint32 Pixel = 0; Pixel < DisplayedFadeOutPixelLength; ++Pixel)
+		FadeOutCurvePoints.SetNumUninitialized(DisplayedFadeOutPixelLength);
+		
+		if (DisplayedFadeOutPixelLength > 0)
 		{
-			const double FadeFraction = static_cast<float>(Pixel + FadeOutPixelOffset) / FadeOutPixelLength;
-			const double CurveFunction = StrongTrimFade->FadeFunctions.FadeOut->GetFadeOutCurveValue(FadeFraction);
+			for (uint32 Pixel = 0; Pixel < DisplayedFadeOutPixelLength; ++Pixel)
+			{
+				const double FadeFraction = static_cast<float>(Pixel) / FadeOutPixelLength;
+				const double CurveFunction = StrongTrimFade->FadeFunctions.FadeOut->GetFadeOutCurveValue(FadeFraction);
 
-			const double CurveValue = Pixel != FadeOutPixelLength - 1 ? 1.f - CurveFunction : 1.f;
+				const bool bNeedsFlickerPrevention = (Pixel == DisplayedFadeOutPixelLength - 1) && DisplayedFadeOutPixelLength != 1;
+
+				// Prevent fade line from flickering if the math doesn't work out to the last Y value being equal to AllottedGeometry.Size.Y as the 
+				// position is changing. This is due to pixel rounding errors. The slight distortion caused by this correction isn't noticeable
+				const double CurveValue = bNeedsFlickerPrevention ? 1.f : 1.f - CurveFunction;
+
+				const uint32 XCoordinate = Pixel + FadeOutStartX;
+				FadeOutCurvePoints[Pixel] = FVector2D(XCoordinate, CurveValue * AllottedGeometry.Size.Y);
+			}
 
-			const uint32 XCoordinate = Pixel + FadeOutStartX;
-			FadeOutCurvePoints[Pixel] = FVector2D(XCoordinate, CurveValue * AllottedGeometry.Size.Y);
+			check(FadeOutCurvePoints.Num() > 0);
+			check(FadeOutCurvePoints[0].Y == 0.f);
 		}
 	}
 	else
