Consider transformation hash design/support for all USoundEFfectSourcePresets

#Synced-CL 48849903

diff --git a/Engine/Plugins/Editor/WaveformEditor/Source/WaveformTransformations/Private/WaveformTransformationEffectChain.cpp b/Engine/Plugins/Editor/WaveformEditor/Source/WaveformTransformations/Private/WaveformTransformationEffectChain.cpp
--- a/Engine/Plugins/Editor/WaveformEditor/Source/WaveformTransformations/Private/WaveformTransformationEffectChain.cpp
+++ b/Engine/Plugins/Editor/WaveformEditor/Source/WaveformTransformations/Private/WaveformTransformationEffectChain.cpp
@@ -102,6 +102,7 @@ Audio::FTransformationPtr UWaveformTransformationEffectChain::CreateTransformati
 			}
 		}
 	}
-	
+PRAGMA_DISABLE_DEPRECATION_WARNINGS
 	return MakeUnique<FWaveTransformationEffectChain>(Effects);
+PRAGMA_ENABLE_DEPRECATION_WARNINGS
 }
diff --git a/Engine/Plugins/Editor/WaveformEditor/Source/WaveformTransformations/Public/WaveformTransformationEffectChain.h b/Engine/Plugins/Editor/WaveformEditor/Source/WaveformTransformations/Public/WaveformTransformationEffectChain.h
--- a/Engine/Plugins/Editor/WaveformEditor/Source/WaveformTransformations/Public/WaveformTransformationEffectChain.h
+++ b/Engine/Plugins/Editor/WaveformEditor/Source/WaveformTransformations/Public/WaveformTransformationEffectChain.h
@@ -10,7 +10,7 @@
 
 #define UE_API WAVEFORMTRANSFORMATIONS_API
 
-class FWaveTransformationEffectChain : public Audio::IWaveTransformation
+class UE_DEPRECATED(5.8, "WaveTransformationEffectChain has been deprecated") FWaveTransformationEffectChain : public Audio::IWaveTransformation
 {
 public:
 	UE_API FWaveTransformationEffectChain(TArray<TObjectPtr<class USoundEffectSourcePreset>>& InEffectPresets);
@@ -24,8 +24,11 @@ private:
 	TArray<TStrongObjectPtr<class USoundEffectSourcePreset>> Presets;
 };
 
+// USoundEffectPresets need diff tag hashes ONLY IF linked to waveform transformations.
+// Removing their connection to waveform transformations is sufficient to avoid the need for cook diff tag support for USoundEffectSourcePresets.
+// This done by deprecating WaveformTransformationEffectChain.
 UCLASS(MinimalAPI, hidden)
-class UWaveformTransformationEffectChain : public UWaveformTransformationBase
+class UE_DEPRECATED(5.8, "WaveformTransformationEffectChain has been deprecated") UWaveformTransformationEffectChain : public UWaveformTransformationBase
 {
 	GENERATED_BODY()
 
