Moving to a synchronous export for automation UE.EditorAutomation(RunTest=Audio) Mac metal-sm5 - TransformationExportTestMonoFileWithMarkers failed

#Synced-CL 49069346

diff --git a/Engine/Plugins/Editor/WaveformEditor/Source/WaveformEditor/Private/WaveformEditor.cpp b/Engine/Plugins/Editor/WaveformEditor/Source/WaveformEditor/Private/WaveformEditor.cpp
--- a/Engine/Plugins/Editor/WaveformEditor/Source/WaveformEditor/Private/WaveformEditor.cpp
+++ b/Engine/Plugins/Editor/WaveformEditor/Source/WaveformEditor/Private/WaveformEditor.cpp
@@ -1447,11 +1447,11 @@ bool FWaveformEditor::BindCommands()
 		FExecuteAction::CreateSP(ZoomManager.ToSharedRef(), &FWaveformEditorZoomController::ZoomOut),
 		FCanExecuteAction::CreateSP(ZoomManager.ToSharedRef(), &FWaveformEditorZoomController::CanZoomOut));
 
-	constexpr bool bPromptUserToSave = true;
+	constexpr bool bIsAsync = true;
 	
 	ToolkitCommands->MapAction(
 		Commands.ExportWaveform,
-		FExecuteAction::CreateSP(this, &FWaveformEditor::ExportWaveform, bPromptUserToSave),
+		FExecuteAction::CreateSP(this, &FWaveformEditor::ExportWaveform, bIsAsync),
 		FCanExecuteAction::CreateSP(WaveWriter.ToSharedRef(), &FWaveformEditorWaveWriter::CanCreateSoundWaveAsset));
 
 	ToolkitCommands->MapAction(
@@ -1910,10 +1910,10 @@ void FWaveformEditor::InitToolMenuContext(FToolMenuContext& MenuContext)
 	MenuContext.AddObject(Context);
 }
 
-void FWaveformEditor::ExportWaveform(bool bPromptUserToSave)
+void FWaveformEditor::ExportWaveform(bool bIsAsync)
 {
 	check(WaveWriter);
-	WaveWriter->ExportTransformedWaveform(bPromptUserToSave);
+	WaveWriter->ExportTransformedWaveform(bIsAsync);
 }
 
 bool FWaveformEditor::CreateDetailsViews()
diff --git a/Engine/Plugins/Editor/WaveformEditor/Source/WaveformEditor/Private/WaveformEditorWaveWriter.cpp b/Engine/Plugins/Editor/WaveformEditor/Source/WaveformEditor/Private/WaveformEditorWaveWriter.cpp
--- a/Engine/Plugins/Editor/WaveformEditor/Source/WaveformEditor/Private/WaveformEditorWaveWriter.cpp
+++ b/Engine/Plugins/Editor/WaveformEditor/Source/WaveformEditor/Private/WaveformEditorWaveWriter.cpp
@@ -21,7 +21,7 @@ bool FWaveformEditorWaveWriter::CanCreateSoundWaveAsset() const
 	return SourceSoundWave != nullptr && WaveWriter != nullptr && WaveWriter->IsDone();
 }
 
-void FWaveformEditorWaveWriter::ExportTransformedWaveform(bool bPromptUserToSave)
+void FWaveformEditorWaveWriter::ExportTransformedWaveform(bool bIsAsync)
 {
 	check(SourceSoundWave);
 	check(WaveWriter);
@@ -41,11 +41,18 @@ void FWaveformEditorWaveWriter::ExportTransformedWaveform(bool bPromptUserToSave
 
 	Audio::TSampleBuffer<> BufferToWrite = GenerateSampleBuffer();
 
-	TFunction<void(const USoundWave*)> OnSoundWaveWritten = [bPromptUserToSave, AssetName, AssetPath](const USoundWave* ResultingWave) {
+	TFunction<void(const USoundWave*)> OnSoundWaveWritten = [AssetName, AssetPath](const USoundWave* ResultingWave) {
 		UE_LOG(LogWaveformEditor, Log, TEXT("Finished Exporting edited soundwave %s/%s"), *AssetPath, *AssetName);
 		TArray<UPackage*> PackagesToSave;
 
 		if (ResultingWave == nullptr)
+		{
+			UE_LOG(LogWaveformEditor, Warning, TEXT("Resulting Exported soundwave is null: %s/%s"), *AssetPath, *AssetName);
+
+			return;
+		}
+
+		if (!ResultingWave->IsValidLowLevel())
 		{
 			UE_LOG(LogWaveformEditor, Warning, TEXT("Resulting Exported soundwave is invalid: %s/%s"), *AssetPath, *AssetName);
 
@@ -71,9 +78,19 @@ void FWaveformEditorWaveWriter::ExportTransformedWaveform(bool bPromptUserToSave
 		}
 	};
 	
-	if (!WaveWriter->BeginWriteToSoundWave(AssetName, BufferToWrite, AssetPath, OnSoundWaveWritten, SourceSoundWave))
+	if (bIsAsync)
 	{
-		UE_LOG(LogWaveformEditor, Log, TEXT("Exporting edited soundwave to %s/%s failed"), *AssetPath, *AssetName);
+		if (!WaveWriter->BeginWriteToSoundWave(AssetName, BufferToWrite, AssetPath, OnSoundWaveWritten, SourceSoundWave))
+		{
+			UE_LOG(LogWaveformEditor, Log, TEXT("Exporting async edited soundwave to %s/%s failed"), *AssetPath, *AssetName);
+		}
+	}
+	else
+	{
+		if (!WaveWriter->SynchronouslyWriteSoundWave(BufferToWrite, &AssetName, &AssetPath, SourceSoundWave))
+		{
+			UE_LOG(LogWaveformEditor, Log, TEXT("Exporting edited soundwave to %s/%s failed"), *AssetPath, *AssetName);
+		}
 	}
 }
 
diff --git a/Engine/Plugins/Editor/WaveformEditor/Source/WaveformEditor/Private/WaveformEditorWaveWriter.h b/Engine/Plugins/Editor/WaveformEditor/Source/WaveformEditor/Private/WaveformEditorWaveWriter.h
--- a/Engine/Plugins/Editor/WaveformEditor/Source/WaveformEditor/Private/WaveformEditorWaveWriter.h
+++ b/Engine/Plugins/Editor/WaveformEditor/Source/WaveformEditor/Private/WaveformEditorWaveWriter.h
@@ -22,7 +22,9 @@ public:
 	explicit FWaveformEditorWaveWriter(USoundWave* InSoundWave);
 
 	bool CanCreateSoundWaveAsset() const;
-	void ExportTransformedWaveform(bool bPromptUserToSave);
+
+	// bIsAsync allows the editor to not hang per file say on bulk operations vs. specific cases that must have sync behaviour like tests
+	void ExportTransformedWaveform(bool bIsAsync);
 
 	WaveformEditorWaveWriter::EChannelFormat GetExportChannelsFormat() const;
 	void SetExportChannelsFormat(const WaveformEditorWaveWriter::EChannelFormat InTargetChannelFormat);
diff --git a/Engine/Plugins/Editor/WaveformEditor/Source/WaveformEditor/Public/WaveformEditor.h b/Engine/Plugins/Editor/WaveformEditor/Source/WaveformEditor/Public/WaveformEditor.h
--- a/Engine/Plugins/Editor/WaveformEditor/Source/WaveformEditor/Public/WaveformEditor.h
+++ b/Engine/Plugins/Editor/WaveformEditor/Source/WaveformEditor/Public/WaveformEditor.h
@@ -52,7 +52,7 @@ public:
 	UE_API bool MatchesContext(const FTransactionContext& InContext, const TArray<TPair<UObject*, FTransactionObjectEvent>>& TransactionObjectContexts) const override;
 	UE_API virtual void InitToolMenuContext(FToolMenuContext& MenuContext) override;
 
-	UE_API void ExportWaveform(bool bPromptUserToSave);
+	UE_API void ExportWaveform(bool bIsAsync);
 
 	// Skips the playhead to the next marker
 	UE_API void SkipToNextMarker();
diff --git a/Engine/Source/Runtime/Engine/Classes/Sound/SampleBufferIO.h b/Engine/Source/Runtime/Engine/Classes/Sound/SampleBufferIO.h
--- a/Engine/Source/Runtime/Engine/Classes/Sound/SampleBufferIO.h
+++ b/Engine/Source/Runtime/Engine/Classes/Sound/SampleBufferIO.h
@@ -178,7 +178,8 @@ namespace Audio
 
 		// This is a blocking call that will return the SoundWave generated from InSampleBuffer.
 		// Optionally, if you're using the editor, you can also write the resulting soundwave out to the content browser using the FileName and FilePath parameters.
-		ENGINE_API USoundWave* SynchronouslyWriteSoundWave(const TSampleBuffer<>& InSampleBuffer, const FString* FileName = nullptr, const FString* FilePath = nullptr);
+		// InSoundWaveTransformationTarget is used to overwrite SoundWave data with its own transformations. (ex. marker data)
+		ENGINE_API USoundWave* SynchronouslyWriteSoundWave(const TSampleBuffer<>& InSampleBuffer, const FString* FileName = nullptr, const FString* FilePath = nullptr, TObjectPtr<USoundWave> InSoundWaveTransformationTarget = nullptr);
 
 		// This blocking call writes out the InSampleBuffer as a wav file at the path specified by FilePath and FileName.
 		// If FilePath is a relative path, it will be relative to the /Saved/BouncedWavFiles folder, otherwise specified absolute path will be used.
diff --git a/Engine/Source/Runtime/Engine/Private/SampleBufferIO.cpp b/Engine/Source/Runtime/Engine/Private/SampleBufferIO.cpp
--- a/Engine/Source/Runtime/Engine/Private/SampleBufferIO.cpp
+++ b/Engine/Source/Runtime/Engine/Private/SampleBufferIO.cpp
@@ -216,6 +216,12 @@ namespace Audio
 	
 		CurrentOperation.Reset(new FAsyncSoundWavePCMWriterTask(this, ESoundWavePCMWriteTaskType::GenerateAndWriteSoundWave, [this, OnSuccess, InSoundWaveTransformationTarget](const USoundWave* InSoundWave)
 		{
+			if (!IsValid(CurrentSoundWave))
+			{
+				UE_LOG(LogAudio, Error, TEXT("CurrentSoundWave is not valid in GenerateAndWriteSoundWave."));
+				return;
+			}
+
 			// Invalidate compressed data on new soundwave so that it gets refreshed or the soundwave is not seekable
 			// and causes this error: FMixerBuffer::CreateStreamingBuffer failed to StreamCompressedInfo on SoundWave
 			CurrentSoundWave->InvalidateCompressedData(true); // calls CurrentSoundWave->FreeResources(), note the game thread is the audio thread in the editor
@@ -319,7 +325,7 @@ namespace Audio
 		return true;
 	}
 
-	USoundWave* FSoundWavePCMWriter::SynchronouslyWriteSoundWave(const TSampleBuffer<>& InSampleBuffer, const FString* FileName, const FString* FilePath)
+	USoundWave* FSoundWavePCMWriter::SynchronouslyWriteSoundWave(const TSampleBuffer<>& InSampleBuffer, const FString* FileName, const FString* FilePath, TObjectPtr<USoundWave> InSoundWaveTransformationTarget)
 	{
 		if (!IsDone())
 		{
@@ -389,6 +395,18 @@ namespace Audio
 		
 		CurrentOperation->StartSynchronousTask();
 
+		// Invalidate compressed data on new soundwave so that it gets refreshed or the soundwave is not seekable
+		// and causes this error: FMixerBuffer::CreateStreamingBuffer failed to StreamCompressedInfo on SoundWave
+		CurrentSoundWave->InvalidateCompressedData(true);
+
+
+#if WITH_EDITOR
+		if (InSoundWaveTransformationTarget && InSoundWaveTransformationTarget->Transformations.Num() > 0)
+		{
+			CurrentSoundWave->OverwriteSoundWaveDataWithTransformationData(InSoundWaveTransformationTarget->Transformations);
+		}
+#endif //WITH_EDITOR
+
 		return CurrentSoundWave;
 	}
 
