Waveform Editor: fix for zero crossing snapping when zoomed

#Synced-CL 49184252

diff --git a/Engine/Plugins/Editor/WaveformEditor/Source/WaveformEditor/Private/WaveformEditor.cpp b/Engine/Plugins/Editor/WaveformEditor/Source/WaveformEditor/Private/WaveformEditor.cpp
--- a/Engine/Plugins/Editor/WaveformEditor/Source/WaveformEditor/Private/WaveformEditor.cpp
+++ b/Engine/Plugins/Editor/WaveformEditor/Source/WaveformEditor/Private/WaveformEditor.cpp
@@ -699,11 +699,9 @@ void FWaveformEditor::SkipToNextZeroCrossing()
 {
 	check(SoundWave != nullptr);
 
-	const int64 StartSample = static_cast<int64>(TransportCoordinator->GetFocusPoint() * SoundWave->TotalSamples);
+	const int64 StartSample = static_cast<int64>(TransportCoordinator->GetPlayBackRatioFromFocusPoint(TransportCoordinator->GetFocusPoint()) * SoundWave->TotalSamples);
 	TArray<float> PCMArray = WaveformView.DataProvider->GetTransformedPCMData();
 
-	const float StartingFocusPoint = TransportCoordinator->GetFocusPoint();
-
 	const UWaveformEditorTransformationsSettings* TransformationsSettings = GetWaveformEditorTransformationsSettings();
 	check(TransformationsSettings);
 
@@ -722,9 +720,10 @@ void FWaveformEditor::SkipToNextZeroCrossing()
 
 	const float ActiveDurationInFrames = ActiveDuration * TransformationChainConfig.SampleRate;
 	const float AdjustedNewPlaybackFramePosition = FMath::Max((StartSample + FramesToMove) - TransformationChainConfig.StartTime * TransformationChainConfig.SampleRate, 0.0f);
+	const float ZoomedRatio = TransportCoordinator->ConvertAbsoluteRatioToZoomed(AdjustedNewPlaybackFramePosition / ActiveDurationInFrames);
 
-	HandlePlayheadScrub(AdjustedNewPlaybackFramePosition / ActiveDurationInFrames, false);
-	TransportCoordinator->SetProgressRatio(AdjustedNewPlaybackFramePosition / ActiveDurationInFrames);
+	HandlePlayheadScrub(ZoomedRatio, false);
+	TransportCoordinator->ScrubFocusPoint(ZoomedRatio, false);
 }
 
 void FWaveformEditor::RegenerateTransformations()
