Waveform Editor: Fix for crash when snapping to zero crossing

#Synced-CL 49936768

diff --git a/Engine/Plugins/Editor/WaveformEditor/Source/WaveformEditor/Private/WaveformEditor.cpp b/Engine/Plugins/Editor/WaveformEditor/Source/WaveformEditor/Private/WaveformEditor.cpp
--- a/Engine/Plugins/Editor/WaveformEditor/Source/WaveformEditor/Private/WaveformEditor.cpp
+++ b/Engine/Plugins/Editor/WaveformEditor/Source/WaveformEditor/Private/WaveformEditor.cpp
@@ -699,15 +699,18 @@ void FWaveformEditor::SkipToNextZeroCrossing()
 {
 	check(SoundWave != nullptr);
 
-	const int64 StartSample = static_cast<int64>(TransportCoordinator->GetPlayBackRatioFromFocusPoint(TransportCoordinator->GetFocusPoint()) * SoundWave->TotalSamples);
 	TArray<float> PCMArray = WaveformView.DataProvider->GetTransformedPCMData();
+	int64 StartFrame = FMath::Min(static_cast<int64>(TransportCoordinator->GetPlayBackRatioFromFocusPoint(TransportCoordinator->GetFocusPoint()) * SoundWave->GetNumFrames()), static_cast<int64>(SoundWave->GetNumFrames()-1));
 
 	const UWaveformEditorTransformationsSettings* TransformationsSettings = GetWaveformEditorTransformationsSettings();
 	check(TransformationsSettings);
 
 	const float NoiseThresholdValue = powf(10.0f, TransformationsSettings->NoiseThreshold / 20.0f);
 
-	const int64 FramesToMove = UWaveformTransformationMarkers::GetFramesToNextZeroCrossing(StartSample, SoundWave->TotalSamples, PCMArray, true, NoiseThresholdValue, SoundWave->NumChannels);
+	check(PCMArray.Num() > 0);
+	check(SoundWave->NumChannels > 0);
+
+	const int64 FramesToMove = UWaveformTransformationMarkers::GetFramesToNextZeroCrossing(StartFrame, SoundWave->GetNumFrames()-1, PCMArray, true, NoiseThresholdValue, SoundWave->NumChannels);
 
 	float ActiveDuration = TransformationChainConfig.EndTime - TransformationChainConfig.StartTime;
 	if (ActiveDuration <= 0.0f)
@@ -719,7 +722,8 @@ void FWaveformEditor::SkipToNextZeroCrossing()
 	check(TransformationChainConfig.SampleRate > 0);
 
 	const float ActiveDurationInFrames = ActiveDuration * TransformationChainConfig.SampleRate;
-	const float AdjustedNewPlaybackFramePosition = FMath::Max((StartSample + FramesToMove) - TransformationChainConfig.StartTime * TransformationChainConfig.SampleRate, 0.0f);
+	const float AdjustedNewPlaybackFramePosition = FMath::Max((StartFrame + FramesToMove) - 
+		(TransformationChainConfig.StartTime * TransformationChainConfig.SampleRate), 0.0f);
 	const float ZoomedRatio = TransportCoordinator->ConvertAbsoluteRatioToZoomed(AdjustedNewPlaybackFramePosition / ActiveDurationInFrames);
 
 	HandlePlayheadScrub(ZoomedRatio, false);
diff --git a/Engine/Plugins/Editor/WaveformEditor/Source/WaveformTransformations/Private/WaveformTransformationMarkers.cpp b/Engine/Plugins/Editor/WaveformEditor/Source/WaveformTransformations/Private/WaveformTransformationMarkers.cpp
--- a/Engine/Plugins/Editor/WaveformEditor/Source/WaveformTransformations/Private/WaveformTransformationMarkers.cpp
+++ b/Engine/Plugins/Editor/WaveformEditor/Source/WaveformTransformations/Private/WaveformTransformationMarkers.cpp
@@ -279,6 +279,8 @@ void UWaveformTransformationMarkers::ModifyMarkerLoopRegion(ELoopModificationCon
 	}
 
 	Markers->Modify();
+	const int64 MaxTotalFrames = static_cast<int64>(NumAvailableSamples) / static_cast<int64>(NumChannels);
+	check(MaxTotalFrames > 0);
 
 	switch (Modification)
 	{
@@ -324,10 +326,10 @@ void UWaveformTransformationMarkers::ModifyMarkerLoopRegion(ELoopModificationCon
 	break;
 	case ELoopModificationControls::LeftHandleNextZeroCrossing:
 	{
-		const int64 StartSample = (LoopRegion->FramePosition) * NumChannels;
-		const int64 EndSample = 0;
+		const int64 StartFrame = (LoopRegion->FramePosition);
+		const int64 EndFrame = 0;
 
-		const int64 FramesToMove = GetFramesToNextZeroCrossing(StartSample, EndSample, PCMArray, false, NoiseThreshold, NumChannels);
+		const int64 FramesToMove = GetFramesToNextZeroCrossing(StartFrame, EndFrame, PCMArray, false, NoiseThreshold, NumChannels);
 
 		LoopRegion->FramePosition -= FramesToMove;
 		LoopRegion->FrameLength += FramesToMove;
@@ -335,10 +337,11 @@ void UWaveformTransformationMarkers::ModifyMarkerLoopRegion(ELoopModificationCon
 	break;
 	case ELoopModificationControls::LeftHandlePreviousZeroCrossing:
 	{
-		const int64 StartSample = (LoopRegion->FramePosition) * NumChannels;
-		const int64 EndSample = (LoopRegion->FramePosition + LoopRegion->FrameLength) * NumChannels;
+		const int64 StartFrame = (LoopRegion->FramePosition);
+		// Clamp end frame to max total frames
+		const int64 EndFrame = FMath::Min((LoopRegion->FramePosition + LoopRegion->FrameLength), MaxTotalFrames-1);
 
-		const int64 FramesToMove = GetFramesToNextZeroCrossing(StartSample, EndSample, PCMArray, true, NoiseThreshold, NumChannels);
+		const int64 FramesToMove = GetFramesToNextZeroCrossing(StartFrame, EndFrame, PCMArray, true, NoiseThreshold, NumChannels);
 
 		LoopRegion->FramePosition += FramesToMove;
 		LoopRegion->FrameLength -= FramesToMove;
@@ -346,25 +349,24 @@ void UWaveformTransformationMarkers::ModifyMarkerLoopRegion(ELoopModificationCon
 	break;
 	case ELoopModificationControls::RightHandleNextZeroCrossing:
 	{
-		const int64 StartSample = (LoopRegion->FramePosition + LoopRegion->FrameLength) * NumChannels;
-		const int64 EndSample = NumAvailableSamples;
+		// Clamp end frame to max total frames
+		const int64 StartFrame = FMath::Min((LoopRegion->FramePosition + LoopRegion->FrameLength), MaxTotalFrames - 1);
+		const int64 EndFrame = MaxTotalFrames - 1;
 
-		LoopRegion->FrameLength += GetFramesToNextZeroCrossing(StartSample, EndSample, PCMArray, true, NoiseThreshold, NumChannels);
+		LoopRegion->FrameLength += GetFramesToNextZeroCrossing(StartFrame, EndFrame, PCMArray, true, NoiseThreshold, NumChannels);
 	}
 	break;
 	case ELoopModificationControls::RightHandlePreviousZeroCrossing:
 	{
-		const int64 StartSample = (LoopRegion->FramePosition + LoopRegion->FrameLength) * NumChannels;
-		const int64 EndSample = LoopRegion->FramePosition * NumChannels;
+		const int64 StartFrame = FMath::Min((LoopRegion->FramePosition + LoopRegion->FrameLength), MaxTotalFrames - 1);
+		const int64 EndFrame = LoopRegion->FramePosition;
 
-		LoopRegion->FrameLength -= GetFramesToNextZeroCrossing(StartSample, EndSample, PCMArray, false, NoiseThreshold, NumChannels);
+		LoopRegion->FrameLength -= GetFramesToNextZeroCrossing(StartFrame, EndFrame, PCMArray, false, NoiseThreshold, NumChannels);
 	}
 	default:
 		break;
 	}
 
-	const int64 MaxTotalFrames = static_cast<int64>(NumAvailableSamples) / static_cast<int64>(NumChannels);
-	check(MaxTotalFrames > 0);
 	// Prevent the frame position from exceeding the wave length
 	LoopRegion->FramePosition = FMath::Clamp(LoopRegion->FramePosition, 0, MaxTotalFrames - 1);
 	// FrameLength can exceed MaxFrames because the loop will be cut by ProcessAudio
@@ -445,49 +447,63 @@ void UWaveformTransformationMarkers::CycleMarkerLoopRegion(ELoopModificationCont
 	OnTransformationChanged.ExecuteIfBound(false);
 }
 
-int64 UWaveformTransformationMarkers::GetFramesToNextZeroCrossing(int64 StartSample, int64 EndSample, const TArray<float>& PCMArray, bool bCheckIncrementing, float NoiseThreshold, int ChannelCount)
+int64 UWaveformTransformationMarkers::GetFramesToNextZeroCrossing(int64 StartFrame, int64 EndFrame, const TArray<float>& PCMArray, bool bCheckIncrementing, float NoiseThreshold, int ChannelCount)
 {
 	check(ChannelCount > 0);
+	check(PCMArray.Num() > 0);
+
+	check(StartFrame >= 0);
+	check(StartFrame <= (PCMArray.Num() / ChannelCount) - 1);
+	check(EndFrame >= 0);
+	check(EndFrame <= (PCMArray.Num() / ChannelCount) - 1);
+
 	int64 FramesMoved = 0;
 
-	bool bAllFramesPositiveLastFrame = false;
-	bool bAllFramesNegativeLastFrame = false;
-	bool bAllFramesLowNoiseLastFrame = false;
+	bool bAllSamplesPositiveLastHighNoise = false;
+	bool bAllSamplesNegativeLastHighNoise = false;
+
+	// Check for clamping within array bounds
+	StartFrame = FMath::Clamp(StartFrame, 0, (PCMArray.Num() / ChannelCount) - 1);
+	EndFrame = FMath::Clamp(EndFrame, 0, (PCMArray.Num() / ChannelCount) - 1);
 
 	if (bCheckIncrementing)
 	{
-		for (int64 i = StartSample; i < EndSample; i += ChannelCount)
+		check(EndFrame >= StartFrame);
+
+		for (int64 i = StartFrame; i < EndFrame; i++)
 		{
-			if (CheckZeroCrossingFrame(i, PCMArray, bAllFramesPositiveLastFrame, bAllFramesNegativeLastFrame, bAllFramesLowNoiseLastFrame, NoiseThreshold, ChannelCount))
+			if (CheckZeroCrossingFrame(i * ChannelCount, PCMArray, bAllSamplesPositiveLastHighNoise, bAllSamplesNegativeLastHighNoise, NoiseThreshold, ChannelCount))
 			{
 				return FramesMoved;
 			}
 
-			++FramesMoved;
+			FramesMoved++;
 		}
 	}
 	else
 	{
-		for (int64 i = StartSample; i >= EndSample; i -= ChannelCount)
+		check(StartFrame >= EndFrame);
+
+		for (int64 i = StartFrame; i >= EndFrame; i--)
 		{
-			if (CheckZeroCrossingFrame(i, PCMArray, bAllFramesPositiveLastFrame, bAllFramesNegativeLastFrame, bAllFramesLowNoiseLastFrame, NoiseThreshold, ChannelCount))
+			if (CheckZeroCrossingFrame(i * ChannelCount, PCMArray, bAllSamplesPositiveLastHighNoise, bAllSamplesNegativeLastHighNoise, NoiseThreshold, ChannelCount))
 			{
 				return FramesMoved;
 			}
 
-			++FramesMoved;
+			FramesMoved++;
 		}
 	}
 
 	return 0;
 }
 
-int64 UWaveformTransformationMarkers::CheckZeroCrossingFrame(int64 CurrentFrame, const TArray<float>& PCMArray, bool& bAllFramesPositiveLastFrame, bool& bAllFramesNegativeLastFrame, bool& bAllFramesLowNoiseLastFrame, float NoiseThreshold, int ChannelCount)
+int64 UWaveformTransformationMarkers::CheckZeroCrossingFrame(int64 CurrentSample, const TArray<float>& PCMArray, bool& bAllSamplesPositiveLastHighNoise, bool& bAllSamplesNegativeLastHighNoise, float NoiseThreshold, int ChannelCount)
 {
 	check(ChannelCount > 0);
-	check(CurrentFrame % ChannelCount == 0); // Check for frame alignment
+	check(CurrentSample % ChannelCount == 0); // Check for frame alignment 
 
-	if (CurrentFrame + ChannelCount > PCMArray.Num())
+	if (CurrentSample + ChannelCount > PCMArray.Num())
 	{
 		return 0;
 	}
@@ -499,7 +515,7 @@ int64 UWaveformTransformationMarkers::CheckZeroCrossingFrame(int64 CurrentFrame,
 	// check if all channels have a zero crossing
 	for (int32 ChannelNum = 0; ChannelNum < ChannelCount; ChannelNum++)
 	{
-		const float FrameValue = PCMArray[CurrentFrame + ChannelNum];
+		const float FrameValue = PCMArray[CurrentSample + ChannelNum];
 
 		// Treat values within the noise threshold as valid for both negative and positive
 		if (FrameValue >= -NoiseThreshold && FrameValue <= NoiseThreshold)
@@ -524,10 +540,13 @@ int64 UWaveformTransformationMarkers::CheckZeroCrossingFrame(int64 CurrentFrame,
 
 	// Check if we have crossed the zero crossing on all channels between last frame and this frame
 	// If we are in a low noise frame then we ignore it for zero crossing
-	bool bZeroCrossingOccurred = (bAllSamplesPositive && bAllFramesNegativeLastFrame && !bAllSamplesLowNoise) || (bAllSamplesNegative && bAllFramesPositiveLastFrame && !bAllSamplesLowNoise);
-	bAllFramesPositiveLastFrame = bAllSamplesPositive;
-	bAllFramesNegativeLastFrame = bAllSamplesNegative;
-	bAllFramesLowNoiseLastFrame = bAllSamplesLowNoise;
+	bool bZeroCrossingOccurred = (bAllSamplesPositive && bAllSamplesNegativeLastHighNoise && !bAllSamplesLowNoise) || (bAllSamplesNegative && bAllSamplesPositiveLastHighNoise && !bAllSamplesLowNoise);
+
+	if (!bAllSamplesLowNoise)
+	{
+		bAllSamplesPositiveLastHighNoise = bAllSamplesPositive;
+		bAllSamplesNegativeLastHighNoise = bAllSamplesNegative;
+	}
 
 	return bZeroCrossingOccurred;
 }
diff --git a/Engine/Plugins/Editor/WaveformEditor/Source/WaveformTransformations/Public/WaveformTransformationMarkers.h b/Engine/Plugins/Editor/WaveformEditor/Source/WaveformTransformations/Public/WaveformTransformationMarkers.h
--- a/Engine/Plugins/Editor/WaveformEditor/Source/WaveformTransformations/Public/WaveformTransformationMarkers.h
+++ b/Engine/Plugins/Editor/WaveformEditor/Source/WaveformTransformations/Public/WaveformTransformationMarkers.h
@@ -131,7 +131,7 @@ public:
 	UE_API void CycleMarkerLoopRegion(ELoopModificationControls Modification);
 
 	static UE_API int64 GetFramesToNextZeroCrossing(int64 StartFrame, int64 EndFrame, const TArray<float>& PCMArray, bool bCheckIncrementing, float NoiseThreshold, int ChannelCount);
-	static UE_API int64 CheckZeroCrossingFrame(int64 CurrentFrame, const TArray<float>& PCMArray, bool& bAllFramesPositiveLastFrame, bool& bAllFramesNegativeLastFrame, bool& bAllFramesLowNoiseLastFrame, float NoiseThreshold, int ChannelCount);
+	static UE_API int64 CheckZeroCrossingFrame(int64 CurrentSample, const TArray<float>& PCMArray, bool& bAllSamplesPositiveLastHighNoise, bool& bAllSamplesNegativeLastHighNoise, float NoiseThreshold, int ChannelCount);
 
 	UE_API bool IsLoopRegionActive() const;
 	UE_API void SetLoopPreviewing();
