Fixes bugs found by BugHawk when trying to decouple fades

#Synced-CL 50346606

diff --git a/Engine/Plugins/Editor/WaveformEditor/Source/WaveformEditorWidgets/Private/WaveformEditorSequenceDataProvider.cpp b/Engine/Plugins/Editor/WaveformEditor/Source/WaveformEditorWidgets/Private/WaveformEditorSequenceDataProvider.cpp
--- a/Engine/Plugins/Editor/WaveformEditor/Source/WaveformEditorWidgets/Private/WaveformEditorSequenceDataProvider.cpp
+++ b/Engine/Plugins/Editor/WaveformEditor/Source/WaveformEditorWidgets/Private/WaveformEditorSequenceDataProvider.cpp
@@ -121,14 +121,16 @@ void FWaveformEditorSequenceDataProvider::GenerateSequenceDataInternal()
 	{
 		Audio::FWaveformTransformationWaveInfo TransformationInfo;
 
+		// Changes in length based on transformations
 		Audio::FAlignedFloatBuffer TransformationsBuffer;
+		// Length remains unchanged 
 		Audio::FAlignedFloatBuffer OutputBuffer;
 
 		TransformationsBuffer.SetNumUninitialized(NumWaveformSamples);
 
 		Audio::ArrayPcm16ToFloat(MakeArrayView((int16*)ImportedRawPCMData.GetData(), NumWaveformSamples), TransformationsBuffer);
 		OutputBuffer = TransformationsBuffer;
-		check(OutputBuffer.Num() > 0);
+		check(OutputBuffer.Num() >= ImportedNumChannels);
 
 		TransformationInfo.Audio = &TransformationsBuffer;
 		TransformationInfo.NumChannels = ImportedNumChannels;
@@ -148,38 +150,39 @@ void FWaveformEditorSequenceDataProvider::GenerateSequenceDataInternal()
 
 				// For rendering, we want to be able to draw over all the samples, not just the ones available after processing
 				RenderLayers[i].Value = FTransformationLayerConstraints(0, 1);
-
 			}
 
-			if (bChainChangesFileLength)
+			if (Transformations[i]->FileChangeLengthPriority() != Audio::ETransformationPriority::None)
 			{
-				FirstEditedSample += TransformationInfo.StartFrameOffset;
-				if (FirstEditedSample >= OutputBuffer.Num())
-				{
-					FirstEditedSample = OutputBuffer.Num() - ImportedNumChannels;
-				}
+				FirstEditedSample += TransformationInfo.StartSampleOffset;
+				
+				check(FirstEditedSample <= OutputBuffer.Num() - ImportedNumChannels);
+				check(FirstEditedSample >= 0);
 
 				//Prevent a lower priority transformation from setting the LastEditedSample past the end of higher priority transformations
 				if (Transformations[i]->FileChangeLengthPriority() == Audio::ETransformationPriority::High
 					|| (Transformations[i]->FileChangeLengthPriority() != Audio::ETransformationPriority::None
-						&& FirstEditedSample + TransformationInfo.NumEditedSamples < LastEditedSample))
-				{
-					LastEditedSample = TransformationInfo.NumEditedSamples <= 0 ? LastEditedSample : FirstEditedSample + TransformationInfo.NumEditedSamples;
-				}
-			
-				int32 EditedBufferSize = LastEditedSample - FirstEditedSample;
-				
-				if (EditedBufferSize <= ImportedNumChannels)
+						&& FirstEditedSample + TransformationsBuffer.Num() < LastEditedSample))
 				{
-					EditedBufferSize = ImportedNumChannels;
-					FirstEditedSample = FMath::Clamp(FirstEditedSample, 0, NumWaveformSamples - ImportedNumChannels);
+					LastEditedSample = FirstEditedSample + TransformationsBuffer.Num();
 				}
-
-				check(FirstEditedSample + EditedBufferSize <= NumWaveformSamples);
-				FMemory::Memcpy(&OutputBuffer.GetData()[FirstEditedSample], TransformationsBuffer.GetData(), EditedBufferSize * sizeof(float));
 			}
 
+			check(TransformationsBuffer.Num() >= ImportedNumChannels);
+			check(TransformationsBuffer.Num() == (LastEditedSample - FirstEditedSample));
+			check(TransformationsBuffer.Num() <= (OutputBuffer.Num() - FirstEditedSample));
+			check(TransformationsBuffer.Num() > 0);
+			check(FirstEditedSample >= 0);
+
+			const int64 DestNum = static_cast<int64>(OutputBuffer.Num()) - FirstEditedSample;
+			const SIZE_T DestSize = static_cast<uint64>(DestNum * sizeof(float));
+			const SIZE_T CopySize = static_cast<uint64>(TransformationsBuffer.Num() * sizeof(float));
+			check(DestSize >= CopySize);
+
+			FMemory::Memcpy(&OutputBuffer.GetData()[FirstEditedSample], TransformationsBuffer.GetData(), CopySize);
+
 			TransformationInfo.NumEditedSamples = 0;
+			TransformationInfo.StartSampleOffset = 0;
 		}
 
 		check(DurationHiglightLayer)
@@ -196,7 +199,9 @@ void FWaveformEditorSequenceDataProvider::GenerateSequenceDataInternal()
 			{
 				TransformedPCMData.SetNumUninitialized(TransformationsBuffer.Num());
 			}
-			FMemory::Memcpy(TransformedPCMData.GetData(), TransformationsBuffer.GetData(), TransformationsBuffer.GetAllocatedSize());
+
+			const SIZE_T CopySize = static_cast<uint64>(TransformationsBuffer.Num() * sizeof(float));
+			FMemory::Memcpy(TransformedPCMData.GetData(), TransformationsBuffer.GetData(), CopySize);
 		}
 		else
 		{
@@ -204,7 +209,9 @@ void FWaveformEditorSequenceDataProvider::GenerateSequenceDataInternal()
 			{
 				TransformedPCMData.SetNumUninitialized(OutputBuffer.Num());
 			}
-			FMemory::Memcpy(TransformedPCMData.GetData(), OutputBuffer.GetData(), OutputBuffer.GetAllocatedSize());
+
+			const SIZE_T CopySize = static_cast<uint64>(OutputBuffer.Num() * sizeof(float));
+			FMemory::Memcpy(TransformedPCMData.GetData(), OutputBuffer.GetData(), CopySize);
 		}
 		
 		const float MaxValue = Audio::ArrayMaxAbsValue(TransformedPCMData);
diff --git a/Engine/Plugins/Editor/WaveformEditor/Source/WaveformTransformations/Private/WaveformTransformationMarkers.cpp b/Engine/Plugins/Editor/WaveformEditor/Source/WaveformTransformations/Private/WaveformTransformationMarkers.cpp
--- a/Engine/Plugins/Editor/WaveformEditor/Source/WaveformTransformations/Private/WaveformTransformationMarkers.cpp
+++ b/Engine/Plugins/Editor/WaveformEditor/Source/WaveformTransformations/Private/WaveformTransformationMarkers.cpp
@@ -75,8 +75,9 @@ void FWaveTransformationMarkers::ProcessAudio(Audio::FWaveformTransformationWave
 	// If InputAudio is smaller than the size of a frame, do not process the audio.
 	check(InOutWaveInfo.NumChannels > 0);
 
-	int64 InputAudioNum = static_cast<int64>(InputAudio.Num());
-	const int64 ExtraSamples = InputAudioNum % InOutWaveInfo.NumChannels;
+	// Note FAlignedFloatBuffer max number of samples is INT32_MAX
+	int32 InputAudioNum = InputAudio.Num();
+	const int32 ExtraSamples = InputAudioNum % InOutWaveInfo.NumChannels;
 	if (ExtraSamples > 0)
 	{
 		InputAudioNum -= ExtraSamples; // trim out samples beyond last full frame
@@ -92,16 +93,16 @@ void FWaveTransformationMarkers::ProcessAudio(Audio::FWaveformTransformationWave
 	check(StartLoopTime >= 0);
 	check(InOutWaveInfo.NumChannels > 0);
 
-	const int64 LastInputAudioIndex = static_cast<int64>(InputAudioNum - 1);
-	const int64 NumChannelsMinusOne = static_cast<int64>(InOutWaveInfo.NumChannels - 1); // Used to step backwards from the end sample to a valid frame
-	const int64 FirstSampleOfLastFrame = LastInputAudioIndex - NumChannelsMinusOne; // The last sample that begins a frame
-	int64 StartSample = FMath::Min(FMath::FloorToInt64(StartLoopTime * InOutWaveInfo.SampleRate) * InOutWaveInfo.NumChannels, FirstSampleOfLastFrame);
-	int64 EndSample = LastInputAudioIndex;
+	const int32 LastInputAudioIndex = InputAudioNum - 1;
+	const int32 NumChannelsMinusOne = InOutWaveInfo.NumChannels - 1; // Used to step backwards from the end sample to a valid frame
+	const int32 FirstSampleOfLastFrame = LastInputAudioIndex - NumChannelsMinusOne; // The last sample that begins a frame
+	int32 StartSample = FMath::Min(FMath::FloorToInt32(StartLoopTime * InOutWaveInfo.SampleRate) * InOutWaveInfo.NumChannels, FirstSampleOfLastFrame);
+	int32 EndSample = LastInputAudioIndex;
 	check(StartSample <= LastInputAudioIndex);
 
 	if (EndLoopTime > 0.f)
 	{
-		const int64 EndFrame = FMath::RoundToInt64(EndLoopTime * InOutWaveInfo.SampleRate);
+		const int32 EndFrame = FMath::RoundToInt32(EndLoopTime * InOutWaveInfo.SampleRate);
 		EndSample = EndFrame * InOutWaveInfo.NumChannels + NumChannelsMinusOne;
 
 		// EndLoopTime can be beyond the length of the file if there is a Trim
@@ -116,8 +117,11 @@ void FWaveTransformationMarkers::ProcessAudio(Audio::FWaveformTransformationWave
 	const int32 FinalSize = (EndSample - StartSample) + 1;
 	check(StartSample + FinalSize <= InputAudioNum);
 	check(FinalSize % InOutWaveInfo.NumChannels == 0);
+	check(StartSample >= 0);
+	// Asserted here because it is validated on load
+	ensure(StartSample % InOutWaveInfo.NumChannels == 0);
 
-	InOutWaveInfo.StartFrameOffset = StartSample - (StartSample % InOutWaveInfo.NumChannels);
+	InOutWaveInfo.StartSampleOffset = StartSample;
 	InOutWaveInfo.NumEditedSamples = FinalSize;
 
 	// Audio needs no trimming
diff --git a/Engine/Plugins/Editor/WaveformEditor/Source/WaveformTransformations/Private/WaveformTransformationNormalize.cpp b/Engine/Plugins/Editor/WaveformEditor/Source/WaveformTransformations/Private/WaveformTransformationNormalize.cpp
--- a/Engine/Plugins/Editor/WaveformEditor/Source/WaveformTransformations/Private/WaveformTransformationNormalize.cpp
+++ b/Engine/Plugins/Editor/WaveformEditor/Source/WaveformTransformations/Private/WaveformTransformationNormalize.cpp
@@ -46,7 +46,7 @@ FWaveTransformationNormalize::FWaveTransformationNormalize(float InTarget, float
 	const float TargetOffset = Target - PeakDecibelValue;
 	const float MakeupGain = FMath::Clamp(TargetOffset, -MaxGain, MaxGain);
 	const float LinearPeakSample = Audio::ArrayMaxAbsValue(InputAudio);
-	InOutWaveInfo.StartFrameOffset = 0;
+	InOutWaveInfo.StartSampleOffset = 0;
 
 	if (MakeupGain != 0.f && LinearPeakSample > 0.f)
 	{
diff --git a/Engine/Plugins/Editor/WaveformEditor/Source/WaveformTransformations/Private/WaveformTransformationTrimFade.cpp b/Engine/Plugins/Editor/WaveformEditor/Source/WaveformTransformations/Private/WaveformTransformationTrimFade.cpp
--- a/Engine/Plugins/Editor/WaveformEditor/Source/WaveformTransformations/Private/WaveformTransformationTrimFade.cpp
+++ b/Engine/Plugins/Editor/WaveformEditor/Source/WaveformTransformations/Private/WaveformTransformationTrimFade.cpp
@@ -16,7 +16,7 @@ namespace WaveformTrasnformationTrimFadePrivate
 
 	constexpr float Linear = 1.f;
 
-	void SetFadeCurve(TObjectPtr<UFadeFunction> FadeFunction, float FadeCurve, float Sharpness)
+	static void SetFadeCurve(TObjectPtr<UFadeFunction> FadeFunction, float FadeCurve, float Sharpness)
 	{
 		if (FadeFunction->IsA<UFadeCurveFunction>())
 		{
@@ -154,8 +154,8 @@ void FWaveTransformationTrimFade::ProcessAudio(Audio::FWaveformTransformationWav
 	// If InputAudio is smaller than the size of a frame, do not process the audio.
 	check(InOutWaveInfo.NumChannels > 0);
 
-	int64 InputAudioNum = static_cast<int64>(InputAudio.Num());
-	const int64 ExtraSamples = InputAudioNum % InOutWaveInfo.NumChannels;
+	int32 InputAudioNum = InputAudio.Num();
+	const int32 ExtraSamples = InputAudioNum % InOutWaveInfo.NumChannels;
 	if (ExtraSamples > 0)
 	{
 		InputAudioNum -= ExtraSamples; // trim out samples beyond last full frame
@@ -168,16 +168,16 @@ void FWaveTransformationTrimFade::ProcessAudio(Audio::FWaveformTransformationWav
 		return;
 	}
 
-	const int64 LastInputAudioIndex = static_cast<int64>(InputAudioNum - 1);
-	const int64 NumChannelsMinusOne = static_cast<int64>(InOutWaveInfo.NumChannels - 1); // Used to step backwards from the end sample to a valid frame
-	const int64 FirstSampleOfLastFrame = LastInputAudioIndex - NumChannelsMinusOne; // The last sample that begins a frame
-	int64 StartSample = FMath::Min(FMath::FloorToInt64(StartTime * InOutWaveInfo.SampleRate) * InOutWaveInfo.NumChannels, FirstSampleOfLastFrame);
-	int64 EndSample = LastInputAudioIndex;
+	const int32 LastInputAudioIndex = InputAudioNum - 1;
+	const int32 NumChannelsMinusOne = InOutWaveInfo.NumChannels - 1; // Used to step backwards from the end sample to a valid frame
+	const int32 FirstSampleOfLastFrame = LastInputAudioIndex - NumChannelsMinusOne; // The last sample that begins a frame
+	int32 StartSample = FMath::Min(FMath::FloorToInt32(StartTime * InOutWaveInfo.SampleRate) * InOutWaveInfo.NumChannels, FirstSampleOfLastFrame);
+	int32 EndSample = LastInputAudioIndex;
 	check(StartSample <= LastInputAudioIndex);
 	
 	if(EndTime > 0.f)
 	{	
-		const int64 EndFrame = FMath::RoundToInt64(EndTime * InOutWaveInfo.SampleRate);
+		const int32 EndFrame = FMath::RoundToInt32(EndTime * InOutWaveInfo.SampleRate);
 		EndSample = EndFrame * InOutWaveInfo.NumChannels + NumChannelsMinusOne;
 		EndSample = FMath::Min(EndSample, LastInputAudioIndex);
 	}
@@ -192,35 +192,29 @@ void FWaveTransformationTrimFade::ProcessAudio(Audio::FWaveformTransformationWav
 	check(EndSample >= StartSample);
 	check((EndSample - NumChannelsMinusOne) % InOutWaveInfo.NumChannels == 0);
 
-	const int64 FinalSize = (EndSample - StartSample) + 1;
-	check(FinalSize % InOutWaveInfo.NumChannels == 0);
-
-	InOutWaveInfo.StartFrameOffset = StartSample - (StartSample % InOutWaveInfo.NumChannels);
-	InOutWaveInfo.NumEditedSamples = FinalSize;
+	const int32 FinalSizeInSamples = (EndSample - StartSample) + 1;
+	check(FinalSizeInSamples % InOutWaveInfo.NumChannels == 0);
 	
-	if (FinalSize > InputAudio.Num() || FinalSize <= 0)
-	{
-		return;
-	}
+	check(FinalSizeInSamples <= InputAudio.Num());
+	check(FinalSizeInSamples > 0);
 
+	const bool bApplyTrim = FinalSizeInSamples < InputAudio.Num();
 	const bool bProcessFadeIn = FadeFunctions.FadeIn && FadeFunctions.FadeIn->Duration > 0.f;
 	const bool bProcessFadeOut = FadeFunctions.FadeOut && FadeFunctions.FadeOut->Duration > 0.f;
-	const bool bProcessFades = bProcessFadeIn || bProcessFadeOut;
 
-	if (!bProcessFades && FinalSize == InputAudio.Num())
+	if (bApplyTrim)
 	{
-		return;
+		FMemory::Memmove(InputAudio.GetData(), &InputAudio[StartSample], FinalSizeInSamples * sizeof(float));
+		InputAudio.SetNum(FinalSizeInSamples, EAllowShrinking::Yes);
+		
+		InOutWaveInfo.StartSampleOffset = StartSample;
+	}
+	else
+	{
+		InOutWaveInfo.StartSampleOffset = 0;
 	}
 
-	TArray<float> TempBuffer;
-	TempBuffer.AddUninitialized(FinalSize);
-
-	FMemory::Memcpy(TempBuffer.GetData(), &InputAudio[StartSample], FinalSize * sizeof(float));
-
-	InputAudio.Empty();
-	InputAudio.AddUninitialized(FinalSize);
-
-	FMemory::Memcpy(InputAudio.GetData(), TempBuffer.GetData(), FinalSize * sizeof(float));
+	InOutWaveInfo.NumEditedSamples = InputAudio.Num();
 
 	if (bProcessFadeIn)
 	{
@@ -312,6 +306,9 @@ void UWaveformTransformationTrimFade::PostEditChangeProperty(FPropertyChangedEve
 {
 	Super::PostEditChangeProperty(PropertyChangedEvent);
 
+	// Transformations are updated in Super::PostEditChangeProperty(PropertyChangedEvent);
+	// When Transformations are updated, UpdateDurationProperties() validates these asserts are correct
+	ensure(EndTime > StartTime);
 	const float MaxDuration = EndTime - StartTime;
 
 	if (FadeFunctions.FadeIn)
@@ -344,6 +341,9 @@ FString UWaveformTransformationTrimFade::GetTransformationHash() const
 			TObjectPtr<UFadeCurveFunction> CurveFunction = CastChecked<UFadeCurveFunction>(FadeFunctions.FadeIn);
 			FadeInCurve = CurveFunction->GetFadeCurve();
 		}
+
+		FPCU::AppendHash(Hash, TEXT("TFIC"), FadeInCurve);
+		FPCU::AppendHash(Hash, TEXT("TFID"), FadeInDuration);
 	}
 
 	if (FadeFunctions.FadeOut)
@@ -354,14 +354,13 @@ FString UWaveformTransformationTrimFade::GetTransformationHash() const
 			TObjectPtr<UFadeCurveFunction> CurveFunction = CastChecked<UFadeCurveFunction>(FadeFunctions.FadeOut);
 			FadeOutCurve = CurveFunction->GetFadeCurve();
 		}
+
+		FPCU::AppendHash(Hash, TEXT("TFOC"), FadeOutCurve);
+		FPCU::AppendHash(Hash, TEXT("TFOD"), FadeOutDuration);
 	}
 
 	FPCU::AppendHash(Hash, TEXT("TFST"), StartTime);
 	FPCU::AppendHash(Hash, TEXT("TFET"), EndTime);
-	FPCU::AppendHash(Hash, TEXT("TFIC"), FadeInCurve);
-	FPCU::AppendHash(Hash, TEXT("TFID"), FadeInDuration);
-	FPCU::AppendHash(Hash, TEXT("TFOC"), FadeOutCurve);
-	FPCU::AppendHash(Hash, TEXT("TFOD"), FadeOutDuration);
 	
 	return Hash;
 }
@@ -511,7 +510,7 @@ PRAGMA_ENABLE_DEPRECATION_WARNINGS
 
 void UWaveformTransformationTrimFade::UpdateDurationProperties(const float InAvailableDuration)
 {
-	check(InAvailableDuration > 0)
+	check(InAvailableDuration > 0.f);
 	AvailableWaveformDuration = InAvailableDuration;
 	StartTime = FMath::Clamp(StartTime, 0.f, AvailableWaveformDuration - UE_KINDA_SMALL_NUMBER);
 	EndTime = EndTime < 0 ? EndTime = AvailableWaveformDuration : FMath::Clamp(EndTime, StartTime + UE_KINDA_SMALL_NUMBER, AvailableWaveformDuration);
diff --git a/Engine/Plugins/Editor/WaveformEditor/Source/WaveformTransformations/Private/WaveformTransformationsModule.cpp b/Engine/Plugins/Editor/WaveformEditor/Source/WaveformTransformations/Private/WaveformTransformationsModule.cpp
--- a/Engine/Plugins/Editor/WaveformEditor/Source/WaveformTransformations/Private/WaveformTransformationsModule.cpp
+++ b/Engine/Plugins/Editor/WaveformEditor/Source/WaveformTransformations/Private/WaveformTransformationsModule.cpp
@@ -12,16 +12,20 @@ DEFINE_LOG_CATEGORY(LogWaveformTransformation);
 void FWaveformTransformationsModule::StartupModule()
 {
 	IModuleInterface::StartupModule();
-	
-	if (FModuleManager::Get().IsModuleLoaded("PropertyEditor"))
+
+	FPropertyEditorModule* PropertyEditorModule = FModuleManager::LoadModulePtr<FPropertyEditorModule>("PropertyEditor");
+
+	if (PropertyEditorModule == nullptr)
 	{
-		FPropertyEditorModule& PropertyEditorModule = FModuleManager::GetModuleChecked<FPropertyEditorModule>("PropertyEditor");
-		
-		PropertyEditorModule.RegisterCustomPropertyTypeLayout(
-			"WaveformTransformationMarkers", 
-			FOnGetPropertyTypeCustomizationInstance::CreateStatic(&FWaveformTransformationMarkersObjectCustomization::MakeInstance)
-		);
+		return;
 	}
+
+	PropertyEditorModule->RegisterCustomPropertyTypeLayout(
+		"WaveformTransformationMarkers",
+		FOnGetPropertyTypeCustomizationInstance::CreateStatic(&FWaveformTransformationMarkersObjectCustomization::MakeInstance)
+	);
+
+	PropertyEditorModule->NotifyCustomizationModuleChanged();
 }
 
 void FWaveformTransformationsModule::ShutdownModule()
@@ -32,7 +36,8 @@ void FWaveformTransformationsModule::ShutdownModule()
 	{
 		FPropertyEditorModule& PropertyEditorModule = FModuleManager::GetModuleChecked<FPropertyEditorModule>("PropertyEditor");
 
-		PropertyEditorModule.UnregisterCustomClassLayout("WaveformTransformationMarkers");
+		PropertyEditorModule.UnregisterCustomPropertyTypeLayout("WaveformTransformationMarkers");
+		PropertyEditorModule.NotifyCustomizationModuleChanged();
 	}
 }
 
diff --git a/Engine/Plugins/Editor/WaveformEditor/Source/WaveformTransformations/Public/WaveformTransformationTrimFade.h b/Engine/Plugins/Editor/WaveformEditor/Source/WaveformTransformations/Public/WaveformTransformationTrimFade.h
--- a/Engine/Plugins/Editor/WaveformEditor/Source/WaveformTransformations/Public/WaveformTransformationTrimFade.h
+++ b/Engine/Plugins/Editor/WaveformEditor/Source/WaveformTransformations/Public/WaveformTransformationTrimFade.h
@@ -164,7 +164,7 @@ public:
 	UPROPERTY(EditAnywhere, Category = "Trim")
 	double EndTime = -1.0;
 
-	UPROPERTY(EditAnywhere, Category = "Fade", meta = (ClampMin = 0.0, DisplayName = "Fade Functions"))
+	UPROPERTY(EditAnywhere, Category = "Fade", meta = (DisplayName = "Fade Functions"))
 	FFadeFunctionData FadeFunctions;
 
 	UE_DEPRECATED(5.7, "Individual fade properties have been deprecated. Use FadeFunctions instead.")
diff --git a/Engine/Plugins/Editor/WaveformEditor/Source/WaveformTransformationsWidgets/Private/WaveformTransformationTrimFadeRenderer.cpp b/Engine/Plugins/Editor/WaveformEditor/Source/WaveformTransformationsWidgets/Private/WaveformTransformationTrimFadeRenderer.cpp
--- a/Engine/Plugins/Editor/WaveformEditor/Source/WaveformTransformationsWidgets/Private/WaveformTransformationTrimFadeRenderer.cpp
+++ b/Engine/Plugins/Editor/WaveformEditor/Source/WaveformTransformationsWidgets/Private/WaveformTransformationTrimFadeRenderer.cpp
@@ -3,7 +3,6 @@
 #include "WaveformTransformationTrimFadeRenderer.h"
 
 #include "AudioWidgetsStyle.h"
-#include "Brushes/SlateRoundedBoxBrush.h"
 #include "Editor.h"
 #include "Framework/Application/SlateApplication.h"
 #include "Layout/Geometry.h"
@@ -19,8 +18,11 @@
 
 FWaveformTransformationTrimFadeRenderer::~FWaveformTransformationTrimFadeRenderer()
 {
-	FSlateApplication::Get().OnFocusChanging().Remove(PopupHandle);
-	FSlateApplication::Get().OnApplicationActivationStateChanged().Remove(ApplicationActivationStateHandle);
+	if (FSlateApplication::IsInitialized())
+	{
+		FSlateApplication::Get().OnFocusChanging().Remove(PopupHandle);
+		FSlateApplication::Get().OnApplicationActivationStateChanged().Remove(ApplicationActivationStateHandle);
+	}
 }
 
 int32 FWaveformTransformationTrimFadeRenderer::OnPaint(const FPaintArgs& Args, const FGeometry& AllottedGeometry, const FSlateRect& MyCullingRect, FSlateWindowElementList& OutDrawElements, int32 LayerId, const FWidgetStyle& InWidgetStyle, bool bParentEnabled) const
@@ -38,7 +40,6 @@ int32 FWaveformTransformationTrimFadeRenderer::DrawTrimHandles(const FGeometry&
 
 	const float HandleWidth = InteractionHandleSize;
 
-	const FSlateRoundedBoxBrush RoundedBoxBrush(FLinearColor::Green, 2);
 	FPaintGeometry HandleGeometry = AllottedGeometry.ToPaintGeometry(FVector2f(HandleWidth, HandleWidth), FSlateLayoutTransform(FVector2f(StartTimeHandleX, AllottedGeometry.Size.Y - HandleWidth)));
 
 	FSlateDrawElement::MakeBox(
@@ -167,7 +168,7 @@ void FWaveformTransformationTrimFadeRenderer::GenerateFadeCurves(const FGeometry
 	FadeInStartX = FMath::RoundToInt32(StartTimeHandleX);
 	FadeOutEndX = FMath::RoundToInt32(EndTimeHandleX);
 
-	if (StrongTrimFade->FadeFunctions.FadeIn && StrongTrimFade->FadeFunctions.FadeIn->Duration > 0.f)
+	if (StrongTrimFade->FadeFunctions.FadeIn && StrongTrimFade->FadeFunctions.FadeIn->Duration > 0.f && PixelsPerFrame > 0.0)
 	{
 		const float FadeInTime = StrongTrimFade->FadeFunctions.FadeIn->Duration;
 		const float FadeInFrames = FadeInTime * TransformationWaveInfo.SampleRate;
@@ -177,7 +178,7 @@ void FWaveformTransformationTrimFadeRenderer::GenerateFadeCurves(const FGeometry
 		
 		FadeInEndX = FMath::RoundToInt32(FMath::Clamp(StartTimeHandleX + FadeInPixelLength, StartTimeHandleX, EndTimeHandleX));
 
-		const uint32 DisplayedFadeInPixelLength = FadeInEndX - FadeInStartX;
+		const uint32 DisplayedFadeInPixelLength = FadeInStartX < FadeInEndX ? FadeInEndX - FadeInStartX : 0;
 		
 		FadeInCurvePoints.SetNumUninitialized(DisplayedFadeInPixelLength);
 
@@ -208,7 +209,7 @@ void FWaveformTransformationTrimFadeRenderer::GenerateFadeCurves(const FGeometry
 		FadeInCurvePoints.SetNumUninitialized(0);
 	}
 
-	if(StrongTrimFade->FadeFunctions.FadeOut && StrongTrimFade->FadeFunctions.FadeOut->Duration > 0.f)
+	if(StrongTrimFade->FadeFunctions.FadeOut && StrongTrimFade->FadeFunctions.FadeOut->Duration > 0.f && PixelsPerFrame > 0.0)
 	{
 		const float FadeOutTime = StrongTrimFade->FadeFunctions.FadeOut->Duration;
 		const float FadeOutFrames = FadeOutTime * TransformationWaveInfo.SampleRate;
@@ -218,7 +219,8 @@ void FWaveformTransformationTrimFadeRenderer::GenerateFadeCurves(const FGeometry
 		check(StartTimeHandleX >= 0);
 		FadeOutStartX = FMath::RoundToInt32(FMath::Clamp(EndTimeHandleX - FadeOutPixelLength, StartTimeHandleX, EndTimeHandleX));
 
-		const uint32 DisplayedFadeOutPixelLength = FadeOutEndX - FadeOutStartX;
+		const uint32 DisplayedFadeOutPixelLength = 
+			(FadeOutStartX < FadeOutEndX) ? (FadeOutEndX - FadeOutStartX) : 0;
 
 		FadeOutCurvePoints.SetNumUninitialized(DisplayedFadeOutPixelLength);
 		
@@ -283,6 +285,13 @@ void FWaveformTransformationTrimFadeRenderer::SetWaveformTransformation(TObjectP
 FReply FWaveformTransformationTrimFadeRenderer::OnMouseWheel(SWidget& OwnerWidget, const FGeometry& MyGeometry, const FPointerEvent& MouseEvent)
 {
 	check(StrongTrimFade);
+
+	// Prevent stomping on other mouse transactions
+	if (TrimFadeInteractionType != ETrimFadeInteractionType::None)
+	{
+		return FReply::Handled();
+	}
+
 	const FVector2D LocalCursorPosition = GetLocalCursorPosition(MouseEvent, MyGeometry);
 
 	if (IsCursorInFadeInInteractionRange(LocalCursorPosition, MyGeometry) && StrongTrimFade->FadeFunctions.FadeIn && StrongTrimFade->FadeFunctions.FadeIn.IsA<UFadeCurveFunction>())
@@ -293,9 +302,15 @@ FReply FWaveformTransformationTrimFadeRenderer::OnMouseWheel(SWidget& OwnerWidge
 		TObjectPtr<UFadeCurveFunction> FadeIn = CastChecked<UFadeCurveFunction>(StrongTrimFade->FadeFunctions.FadeIn);
 		OutPropertyName = FadeIn->GetFadeCurvePropertyName();
 
-		if (BeginTransaction(TEXT("PropertyEditor"), FText::Format(NSLOCTEXT("PropertyEditor", "SetNumericPropertyTransaction", "Edit {0}"), OutPropertyName), StrongTrimFade.Get()) == INDEX_NONE)
+		if (!GIsTransacting)
 		{
-			UE_LOG(LogWaveformTransformation, Warning, TEXT("Begin TrimFadeTransformation OnMouseWheel Transaction Failed"));
+			TransactionResult =
+				BeginTransaction(TEXT("PropertyEditor"), FText::Format(NSLOCTEXT("PropertyEditor", "SetNumericPropertyTransaction", "Edit {0}"), OutPropertyName), StrongTrimFade.Get());
+
+			if (TransactionResult == INDEX_NONE)
+			{
+				UE_LOG(LogWaveformTransformation, Warning, TEXT("Begin TrimFadeTransformation OnMouseWheel Transaction Failed"));
+			}
 		}
 
 		StrongTrimFade->Modify();
@@ -304,7 +319,11 @@ FReply FWaveformTransformationTrimFadeRenderer::OnMouseWheel(SWidget& OwnerWidge
 		const float FadeCurveValue = FadeIn->GetFadeCurve() + FadeCurveDelta;
 		FadeIn->SetFadeCurve(FadeCurveValue);
 
-		EndTransaction();
+		if (TransactionResult != INDEX_NONE)
+		{
+			EndTransaction();
+			TransactionResult = INDEX_NONE;
+		}
 		StrongTrimFade->OnTransformationChanged.ExecuteIfBound(true);
 
 		return FReply::Handled();
@@ -317,9 +336,17 @@ FReply FWaveformTransformationTrimFadeRenderer::OnMouseWheel(SWidget& OwnerWidge
 
 		TObjectPtr<UFadeCurveFunction> FadeOut = CastChecked<UFadeCurveFunction>(StrongTrimFade->FadeFunctions.FadeOut);
 		OutPropertyName = FadeOut->GetFadeCurvePropertyName();
-		if (BeginTransaction(TEXT("PropertyEditor"), FText::Format(NSLOCTEXT("PropertyEditor", "SetNumericPropertyTransaction", "Edit {0}"), OutPropertyName), StrongTrimFade.Get()) == INDEX_NONE)
+
+		if (!GIsTransacting)
 		{
-			UE_LOG(LogWaveformTransformation, Warning, TEXT("Begin TrimFadeTransformation OnMouseWheel Transaction Failed"));
+			TransactionResult =
+				BeginTransaction(TEXT("PropertyEditor"), FText::Format(NSLOCTEXT("PropertyEditor", "SetNumericPropertyTransaction", "Edit {0}"), OutPropertyName),
+					StrongTrimFade.Get());
+			
+			if (TransactionResult == INDEX_NONE)
+			{
+				UE_LOG(LogWaveformTransformation, Warning, TEXT("Begin TrimFadeTransformation OnMouseWheel Transaction Failed"));
+			}
 		}
 
 		StrongTrimFade->Modify();
@@ -328,7 +355,11 @@ FReply FWaveformTransformationTrimFadeRenderer::OnMouseWheel(SWidget& OwnerWidge
 		const float FadeCurveValue = FadeOut->GetFadeCurve() + FadeCurveDelta;
 		FadeOut->SetFadeCurve(FadeCurveValue);
 
-		EndTransaction();
+		if (TransactionResult != INDEX_NONE)
+		{
+			EndTransaction();
+			TransactionResult = INDEX_NONE;
+		}
 		StrongTrimFade->OnTransformationChanged.ExecuteIfBound(true);
 
 		return FReply::Handled();
@@ -346,6 +377,7 @@ FVector2D FWaveformTransformationTrimFadeRenderer::GetLocalCursorPosition(const
 double FWaveformTransformationTrimFadeRenderer::ConvertXRatioToTime(const float InRatio) const
 {
 	check(TransformationWaveInfo.NumChannels > 0);
+	check(TransformationWaveInfo.SampleRate > 0);
 
 	const float NumFrames = TransformationWaveInfo.TotalNumSamples / TransformationWaveInfo.NumChannels;
 	const float FrameSelected = NumFrames * InRatio;
@@ -386,10 +418,10 @@ void FWaveformTransformationTrimFadeRenderer::UpdateInteractionRange(const FGeom
 	EndTimeInteractionYRange.SetLowerBoundValue(AllottedGeometry.GetLocalSize().Y - InteractionHandleSize);
 	EndTimeInteractionYRange.SetUpperBoundValue(AllottedGeometry.GetLocalSize().Y);
 
-	FadeInInteractionXRange.SetLowerBoundValue(FadeInEndX - InteractionHandleSize);
-	FadeInInteractionXRange.SetUpperBoundValue(FadeInEndX + InteractionHandleSize);
-	FadeOutInteractionXRange.SetLowerBoundValue(FadeOutStartX - InteractionHandleSize);
-	FadeOutInteractionXRange.SetUpperBoundValue(FadeOutStartX + InteractionHandleSize);
+	FadeInInteractionXRange.SetLowerBoundValue(static_cast<float>(FadeInEndX) - InteractionHandleSize);
+	FadeInInteractionXRange.SetUpperBoundValue(static_cast<float>(FadeInEndX) + InteractionHandleSize);
+	FadeOutInteractionXRange.SetLowerBoundValue(static_cast<float>(FadeOutStartX) - InteractionHandleSize);
+	FadeOutInteractionXRange.SetUpperBoundValue(static_cast<float>(FadeOutStartX) + InteractionHandleSize);
 }
 
 FReply FWaveformTransformationTrimFadeRenderer::OnMouseButtonDown(SWidget& OwnerWidget, const FGeometry& MyGeometry, const FPointerEvent& MouseEvent)
@@ -403,9 +435,15 @@ FReply FWaveformTransformationTrimFadeRenderer::OnMouseButtonDown(SWidget& Owner
 
 	if (TrimFadeInteractionType != ETrimFadeInteractionType::None)
 	{
-		if (BeginTransaction(TEXT("PropertyEditor"), FText::Format(NSLOCTEXT("PropertyEditor", "SetNumericPropertyTransaction", "Edit {0}"), GetPropertyEditedByCurrentInteraction()), StrongTrimFade.Get()) == INDEX_NONE)
+		if (!GIsTransacting)
 		{
-			UE_LOG(LogWaveformTransformation, Warning, TEXT("Begin TrimFadeTransformation OnMouseButtonDown Transaction Failed"));
+			TransactionResult =
+				BeginTransaction(TEXT("PropertyEditor"), FText::Format(NSLOCTEXT("PropertyEditor", "SetNumericPropertyTransaction", "Edit {0}"), GetPropertyEditedByCurrentInteraction()), StrongTrimFade.Get());
+
+			if (TransactionResult == INDEX_NONE)
+			{
+				UE_LOG(LogWaveformTransformation, Warning, TEXT("Begin TrimFadeTransformation OnMouseButtonDown Transaction Failed"));
+			}
 		}
 
 		if (TrimFadeInteractionType == ETrimFadeInteractionType::ScrubbingFadeIn && StrongTrimFade->FadeFunctions.FadeIn)
@@ -435,8 +473,6 @@ FReply FWaveformTransformationTrimFadeRenderer::OnMouseMove(SWidget& OwnerWidget
 
 		StrongTrimFade->PostEditChange();
 
-		StrongTrimFade->OnTransformationRenderChanged.ExecuteIfBound();
-
 		return FReply::Handled().CaptureMouse(OwnerWidget.AsShared());
 	}
 
@@ -450,7 +486,11 @@ FReply FWaveformTransformationTrimFadeRenderer::OnMouseButtonUp(SWidget& OwnerWi
 	{
 		SetPropertyValueDependingOnInteractionType(MouseEvent, MyGeometry);
 
-		EndTransaction();
+		if (TransactionResult != INDEX_NONE)
+		{
+			EndTransaction();
+			TransactionResult = INDEX_NONE;
+		}
 
 		StrongTrimFade->PostEditChange();
 
@@ -467,6 +507,8 @@ FReply FWaveformTransformationTrimFadeRenderer::OnMouseButtonUp(SWidget& OwnerWi
 void FWaveformTransformationTrimFadeRenderer::SetPropertyValueDependingOnInteractionType(const FPointerEvent& MouseEvent, const FGeometry& WidgetGeometry)
 {
 	check(StrongTrimFade);
+	check(WidgetGeometry.GetLocalSize().X > 0.f);
+	check(WidgetGeometry.GetLocalSize().X > 0.f);
 	const FVector2D LocalCursorPosition = GetLocalCursorPosition(MouseEvent, WidgetGeometry);
 	const float LocalCursorXRatio = FMath::Clamp(LocalCursorPosition.X / WidgetGeometry.GetLocalSize().X, 0.f, 1.f);
 	const double SelectedTime = ConvertXRatioToTime(LocalCursorXRatio);
@@ -514,6 +556,13 @@ void FWaveformTransformationTrimFadeRenderer::SetPropertyValueDependingOnInterac
 
 void FWaveformTransformationTrimFadeRenderer::ShowSelectFadeModeMenuAtCursor(const FGeometry& MyGeometry, const FPointerEvent& MouseEvent)
 {
+	if (!FSlateApplication::IsInitialized())
+	{
+		UE_LOG(LogWaveformTransformation, Warning, TEXT("FSlateApplication not initialized in WaveformTransformationTrimFadeRenderer."));
+
+		return;
+	}
+
 	const FVector2D LocalWindowMaxPosition = MyGeometry.GetAbsolutePosition() + MyGeometry.GetAbsoluteSize();
 	const FVector2D LocalCursorPosition = MouseEvent.GetScreenSpacePosition();
 
@@ -540,6 +589,51 @@ void FWaveformTransformationTrimFadeRenderer::ShowSelectFadeModeMenuAtCursor(con
 	FadeModeMenuWindow = MenuWindow;
 	TWeakObjectPtr<UWaveformTransformationTrimFade> InWeakTrimFade = StrongTrimFade.Get();
 
+	TSharedPtr<FDelegateHandle> ApplicationActivationStateHandlePtr = MakeShared<FDelegateHandle>();
+	*ApplicationActivationStateHandlePtr = FSlateApplication::Get().OnApplicationActivationStateChanged().AddLambda([InFadeModeMenuWindow = FadeModeMenuWindow, ApplicationActivationStateHandlePtr](bool isActive)
+		{
+			if (!isActive)
+			{
+				if (TSharedPtr<SWindow> LockedFadeModeMenuWindow = InFadeModeMenuWindow.Pin())
+				{
+					FSlateApplication::Get().RequestDestroyWindow(LockedFadeModeMenuWindow.ToSharedRef());
+					FSlateApplication::Get().OnApplicationActivationStateChanged().Remove(*ApplicationActivationStateHandlePtr);
+				}
+			}
+		});
+
+	if (ApplicationActivationStateHandle.IsValid())
+	{
+		FSlateApplication::Get().OnApplicationActivationStateChanged().Remove(ApplicationActivationStateHandle);
+	}
+
+	ApplicationActivationStateHandle = *ApplicationActivationStateHandlePtr;
+
+	// If focus is lost on the popup, destroy it to prevent popups hanging around
+	TSharedPtr<FDelegateHandle> PopupHandlePtr = MakeShared<FDelegateHandle>();
+	*PopupHandlePtr = FSlateApplication::Get().OnFocusChanging().AddLambda([InFadeModeMenuWindow = FadeModeMenuWindow, PopupHandlePtr](const FFocusEvent& FocusEvent, const FWeakWidgetPath& WeakWidgetPath
+		, const TSharedPtr<SWidget>& OldWidget, const FWidgetPath& WidgetPath, const TSharedPtr<SWidget>& NewWidget)
+		{
+			if (InFadeModeMenuWindow != nullptr && InFadeModeMenuWindow.IsValid())
+			{
+				if (TSharedPtr<SWindow> LockedFadeModeMenuWindow = InFadeModeMenuWindow.Pin())
+				{
+					if (OldWidget && !OldWidget->IsHovered() && LockedFadeModeMenuWindow == OldWidget)
+					{
+						FSlateApplication::Get().RequestDestroyWindow(LockedFadeModeMenuWindow.ToSharedRef());
+						FSlateApplication::Get().OnFocusChanging().Remove(*PopupHandlePtr);
+					}
+				}
+			}
+		});
+
+	if (PopupHandle.IsValid())
+	{
+		FSlateApplication::Get().OnFocusChanging().Remove(PopupHandle);
+	}
+
+	PopupHandle = *PopupHandlePtr;
+
 	for (const TPair<EWaveEditorFadeMode, TSubclassOf<UFadeFunction>>& FadeOptionPair : UWaveformTransformationTrimFade::FadeModeToFadeFunctionMap)
 	{
 		MenuContent->AddSlot()
@@ -547,7 +641,7 @@ void FWaveformTransformationTrimFadeRenderer::ShowSelectFadeModeMenuAtCursor(con
 			.AutoHeight()
 			[
 				SNew(SButton)
-					.OnClicked_Lambda([InWeakFadeModeMenuWindow = FadeModeMenuWindow, InWeakTrimFade, InteractionType = TrimFadeInteractionType, FadeOptionPair, PropertyName = GetPropertyEditedByCurrentInteraction()]() -> FReply
+					.OnClicked_Lambda([InPopupHandle = PopupHandle, InApplicationActivationStateHandle = ApplicationActivationStateHandle, InWeakFadeModeMenuWindow = FadeModeMenuWindow, InWeakTrimFade, InteractionType = TrimFadeInteractionType, FadeOptionPair, PropertyName = GetPropertyEditedByCurrentInteraction()]() -> FReply
 					{
 						if (TStrongObjectPtr<UWaveformTransformationTrimFade> PinnedTrimFade = InWeakTrimFade.Pin())
 						{
@@ -590,6 +684,16 @@ void FWaveformTransformationTrimFadeRenderer::ShowSelectFadeModeMenuAtCursor(con
 						if (LockedFadeModeMenuWindow)
 						{
 							FSlateApplication::Get().RequestDestroyWindow(LockedFadeModeMenuWindow.ToSharedRef());
+
+							if (InPopupHandle.IsValid())
+							{
+								FSlateApplication::Get().OnFocusChanging().Remove(InPopupHandle);
+							}
+
+							if (InApplicationActivationStateHandle.IsValid())
+							{
+								FSlateApplication::Get().OnApplicationActivationStateChanged().Remove(InApplicationActivationStateHandle);
+							}
 						}
 
 						return FReply::Handled();
@@ -635,51 +739,6 @@ void FWaveformTransformationTrimFadeRenderer::ShowSelectFadeModeMenuAtCursor(con
 
 			return false; // return false to only execute once
 		}));
-
-	TSharedPtr<FDelegateHandle> ApplicationActivationStateHandlePtr = MakeShared<FDelegateHandle>();
-	*ApplicationActivationStateHandlePtr = FSlateApplication::Get().OnApplicationActivationStateChanged().AddLambda([InFadeModeMenuWindow = FadeModeMenuWindow, ApplicationActivationStateHandlePtr](bool isActive)
-		{
-			if (!isActive)
-			{
-				if (TSharedPtr<SWindow> LockedFadeModeMenuWindow = InFadeModeMenuWindow.Pin())
-				{
-					FSlateApplication::Get().RequestDestroyWindow(LockedFadeModeMenuWindow.ToSharedRef());
-					FSlateApplication::Get().OnApplicationActivationStateChanged().Remove(*ApplicationActivationStateHandlePtr);
-				}
-			}
-		});
-
-	if (ApplicationActivationStateHandle.IsValid())
-	{
-		FSlateApplication::Get().OnApplicationActivationStateChanged().Remove(ApplicationActivationStateHandle);
-	}
-
-	ApplicationActivationStateHandle = *ApplicationActivationStateHandlePtr;
-
-	// If focus is lost on the popup, destroy it to prevent popups hanging around
-	TSharedPtr<FDelegateHandle> PopupHandlePtr = MakeShared<FDelegateHandle>();
-	*PopupHandlePtr = FSlateApplication::Get().OnFocusChanging().AddLambda([InFadeModeMenuWindow = FadeModeMenuWindow, PopupHandlePtr](const FFocusEvent& FocusEvent, const FWeakWidgetPath& WeakWidgetPath
-		, const TSharedPtr<SWidget>& OldWidget, const FWidgetPath& WidgetPath, const TSharedPtr<SWidget>& NewWidget)
-		{
-			if (InFadeModeMenuWindow != nullptr && InFadeModeMenuWindow.IsValid())
-			{
-				if (TSharedPtr<SWindow> LockedFadeModeMenuWindow = InFadeModeMenuWindow.Pin())
-				{
-					if (OldWidget && !OldWidget->IsHovered() && LockedFadeModeMenuWindow == OldWidget)
-					{
-						FSlateApplication::Get().RequestDestroyWindow(LockedFadeModeMenuWindow.ToSharedRef());
-						FSlateApplication::Get().OnFocusChanging().Remove(*PopupHandlePtr);
-					}
-				}
-			}
-		});
-
-	if (PopupHandle.IsValid())
-	{
-		FSlateApplication::Get().OnFocusChanging().Remove(PopupHandle);
-	}
-
-	PopupHandle = *PopupHandlePtr;
 }
 
 const FSlateIcon FWaveformTransformationTrimFadeRenderer::GetFadeModeIcon(const EWaveEditorFadeMode& FadeMode)
diff --git a/Engine/Plugins/Editor/WaveformEditor/Source/WaveformTransformationsWidgets/Private/WaveformTransformationTrimFadeRenderer.h b/Engine/Plugins/Editor/WaveformEditor/Source/WaveformTransformationsWidgets/Private/WaveformTransformationTrimFadeRenderer.h
--- a/Engine/Plugins/Editor/WaveformEditor/Source/WaveformTransformationsWidgets/Private/WaveformTransformationTrimFadeRenderer.h
+++ b/Engine/Plugins/Editor/WaveformEditor/Source/WaveformTransformationsWidgets/Private/WaveformTransformationTrimFadeRenderer.h
@@ -5,10 +5,12 @@
 #include "WaveformTransformationRendererBase.h"
 #include "WaveformTransformationTrimFade.h"
 
+#include "Brushes/SlateRoundedBoxBrush.h"
+
 class FWaveformTransformationTrimFadeRenderer : public FWaveformTransformationRendererBase
 {
 public:
-	FWaveformTransformationTrimFadeRenderer() = default;
+	FWaveformTransformationTrimFadeRenderer() : RoundedBoxBrush(FLinearColor::Green, 2) {}
 	~FWaveformTransformationTrimFadeRenderer();
 
 	virtual int32 OnPaint(const FPaintArgs& Args, const FGeometry& AllottedGeometry, const FSlateRect& MyCullingRect, FSlateWindowElementList& OutDrawElements, int32 LayerId, const FWidgetStyle& InWidgetStyle, bool bParentEnabled) const override;
@@ -53,6 +55,7 @@ private:
 	double PixelsPerFrame = 0.0;
 	const float InteractionHandleSize = 30.0f;
 	FVector2D MousePosition;
+	int32 TransactionResult = INDEX_NONE;
 
 	enum class ETrimFadeInteractionType : uint8
 	{
@@ -70,6 +73,8 @@ private:
 	void ShowSelectFadeModeMenuAtCursor(const FGeometry& MyGeometry, const FPointerEvent& MouseEvent);
 	const FSlateIcon GetFadeModeIcon(const EWaveEditorFadeMode& FadeMode);
 
+	const FSlateRoundedBoxBrush RoundedBoxBrush;
+
 	TWeakPtr<SWindow> FadeModeMenuWindow = nullptr;
 
 	TStrongObjectPtr<UWaveformTransformationTrimFade> StrongTrimFade = nullptr;
diff --git a/Engine/Source/Runtime/AudioExtensions/Private/IWaveformTransformation.cpp b/Engine/Source/Runtime/AudioExtensions/Private/IWaveformTransformation.cpp
--- a/Engine/Source/Runtime/AudioExtensions/Private/IWaveformTransformation.cpp
+++ b/Engine/Source/Runtime/AudioExtensions/Private/IWaveformTransformation.cpp
@@ -47,3 +47,16 @@ void UWaveformTransformationBase::NotifyPropertyChange(FProperty* Property)
 	}
 }
 #endif //WITH_EDITOR
+
+int64 FWaveTransformUObjectConfiguration::GetEndFrameOffset() const
+{
+	if (EndTime < 0.f || OriginalNumFrames <= 0 || SampleRate <= 0.f)
+	{
+		return 0;
+	}
+
+	int64 EndTimeInFrames = FMath::RoundToInt64(EndTime * SampleRate);
+	
+	// Clamp to OriginalNumFrames to prevent a negative offset
+	return FMath::Clamp((OriginalNumFrames - EndTimeInFrames), 0, OriginalNumFrames);
+}
diff --git a/Engine/Source/Runtime/AudioExtensions/Public/IWaveformTransformation.h b/Engine/Source/Runtime/AudioExtensions/Public/IWaveformTransformation.h
--- a/Engine/Source/Runtime/AudioExtensions/Public/IWaveformTransformation.h
+++ b/Engine/Source/Runtime/AudioExtensions/Public/IWaveformTransformation.h
@@ -24,8 +24,8 @@ namespace Audio
 		float SampleRate = 0.f;
 		int32 NumChannels = 0;
 		Audio::FAlignedFloatBuffer* Audio = nullptr;
-		uint32 StartFrameOffset = 0;
-		uint32 NumEditedSamples = 0;
+		int32 StartSampleOffset = 0;
+		int32 NumEditedSamples = 0;
 	};
 
 	enum class ETransformationPriority : uint8
@@ -122,6 +122,16 @@ struct FWaveTransformUObjectConfiguration
 	TArray<FSoundWaveCuePoint> WaveCues; // List of cues parsed from the wave file
 	bool bIsPreviewingLoopRegion = false;
 	bool bCachedSoundWaveLoopState = false;
+
+protected:
+	int64 OriginalNumFrames = 0;
+
+public:
+	void SetOriginalNumFrames(const int64 InOriginalNumFrames) { OriginalNumFrames = InOriginalNumFrames; }
+	const int64 GetOriginalNumFrames() const { return OriginalNumFrames; }
+
+	int64 GetStartFrameOffset() const { return FMath::RoundToInt64(StartTime * SampleRate); }
+	AUDIOEXTENSIONS_API int64 GetEndFrameOffset() const;
 };
 
 #if WITH_EDITORONLY_DATA
@@ -148,7 +158,7 @@ public:
 	// Sort to ensure proper order of operation for audio processing
 	static void SortTransformationsArray(TArray<TObjectPtr<UWaveformTransformationBase>>& InOutTransformations) 
 	{
-		Algo::Sort(InOutTransformations, [](const TObjectPtr<UWaveformTransformationBase>& TransformationA, const TObjectPtr<UWaveformTransformationBase>& TransformationB)
+		Algo::StableSort(InOutTransformations, [](const TObjectPtr<UWaveformTransformationBase>& TransformationA, const TObjectPtr<UWaveformTransformationBase>& TransformationB)
 			{
 				if ((!TransformationA && TransformationB) || (!TransformationA && !TransformationB))
 				{
diff --git a/Engine/Source/Runtime/Engine/Private/SoundWave.cpp b/Engine/Source/Runtime/Engine/Private/SoundWave.cpp
--- a/Engine/Source/Runtime/Engine/Private/SoundWave.cpp
+++ b/Engine/Source/Runtime/Engine/Private/SoundWave.cpp
@@ -4911,6 +4911,7 @@ const FWaveTransformUObjectConfiguration& USoundWave::GetTransformationChainConf
 
 const FWaveTransformUObjectConfiguration& USoundWave::UpdateTransformations()
 {
+	// All TransformationChainConfig are initialized in this UpdateTransformations() and can be modified by transformations
 	TransformationChainConfig = FWaveTransformUObjectConfiguration();
 
 	if (TotalSamples == 0)
@@ -4927,6 +4928,7 @@ const FWaveTransformUObjectConfiguration& USoundWave::UpdateTransformations()
 	TransformationChainConfig.WaveCues = CuePoints;
 	TransformationChainConfig.bIsPreviewingLoopRegion = false;
 	TransformationChainConfig.bCachedSoundWaveLoopState = static_cast<bool>(bLooping);
+	TransformationChainConfig.SetOriginalNumFrames(Duration * TransformationChainConfig.SampleRate);
 
 	TSet<UClass*> TransformationClassSet;
 	for (int32 i = 0; i < Transformations.Num(); ++i)
