WaveformEditor: Fix for zero crossing snapping getting stuck on certain frames

#Synced-CL 50416702

diff --git a/Engine/Plugins/Editor/WaveformEditor/Source/WaveformEditor/Private/WaveformEditor.cpp b/Engine/Plugins/Editor/WaveformEditor/Source/WaveformEditor/Private/WaveformEditor.cpp
--- a/Engine/Plugins/Editor/WaveformEditor/Source/WaveformEditor/Private/WaveformEditor.cpp
+++ b/Engine/Plugins/Editor/WaveformEditor/Source/WaveformEditor/Private/WaveformEditor.cpp
@@ -700,7 +700,7 @@ void FWaveformEditor::SkipToNextZeroCrossing()
 	check(SoundWave != nullptr);
 
 	TArray<float> PCMArray = WaveformView.DataProvider->GetTransformedPCMData();
-	int64 StartFrame = FMath::Min(static_cast<int64>(TransportCoordinator->GetPlayBackRatioFromFocusPoint(TransportCoordinator->GetFocusPoint()) * SoundWave->GetNumFrames()), static_cast<int64>(SoundWave->GetNumFrames()-1));
+	int64 StartFrame = FMath::Min(FMath::RoundToInt64(TransportCoordinator->GetPlayBackRatioFromFocusPoint(TransportCoordinator->GetFocusPoint()) * SoundWave->GetNumFrames()), static_cast<int64>(SoundWave->GetNumFrames()-1));
 
 	const UWaveformEditorTransformationsSettings* TransformationsSettings = GetWaveformEditorTransformationsSettings();
 	check(TransformationsSettings);
