Decouple fades from trim

#Synced-CL 50650077

diff --git a/Engine/Plugins/Editor/WaveformEditor/Source/WaveformEditor/Private/WaveformEditor.cpp b/Engine/Plugins/Editor/WaveformEditor/Source/WaveformEditor/Private/WaveformEditor.cpp
--- a/Engine/Plugins/Editor/WaveformEditor/Source/WaveformEditor/Private/WaveformEditor.cpp
+++ b/Engine/Plugins/Editor/WaveformEditor/Source/WaveformEditor/Private/WaveformEditor.cpp
@@ -193,6 +193,7 @@ FWaveformEditor::~FWaveformEditor()
 	{
 		// If the user has not used the trimfade then remove the transformation to avoid unnecessary transformations in the chain
 		check(TrimFadeTransformation != nullptr);
+		check(TrimFadeTransformation->FadeTransformation != nullptr);
 		check(SoundWave != nullptr);
 
 		// Account for a user trimming a single frame by having a 1/2 frame error tolerance
@@ -203,11 +204,13 @@ FWaveformEditor::~FWaveformEditor()
 
 		if (FMath::IsNearlyEqual(TrimFadeTransformation->StartTime, 0, ErrorTolerance) && FMath::IsNearlyEqual(TrimFadeTransformation->EndTime, SoundWave->Duration, ErrorTolerance))
 		{
-			if ((TrimFadeTransformation->FadeFunctions.FadeIn && FMath::IsNearlyEqual(TrimFadeTransformation->FadeFunctions.FadeIn->Duration, 0, ErrorTolerance)
-				&& TrimFadeTransformation->FadeFunctions.FadeOut && FMath::IsNearlyEqual(TrimFadeTransformation->FadeFunctions.FadeOut->Duration, 0, ErrorTolerance))
-				|| (TrimFadeTransformation->FadeFunctions.FadeIn == nullptr && TrimFadeTransformation->FadeFunctions.FadeOut == nullptr))
+			if (TrimFadeTransformation->FadeTransformation->GetFadeRegionsNum() != 1
+				|| (TrimFadeTransformation->GetFadeRegion().FadeIn && FMath::IsNearlyEqual(TrimFadeTransformation->GetFadeRegion().FadeIn->Duration, 0, ErrorTolerance)
+					&& TrimFadeTransformation->GetFadeRegion().FadeOut && FMath::IsNearlyEqual(TrimFadeTransformation->GetFadeRegion().FadeOut->Duration, 0, ErrorTolerance))
+				|| (TrimFadeTransformation->GetFadeRegion().FadeIn == nullptr && TrimFadeTransformation->GetFadeRegion().FadeOut == nullptr))
 			{
-				// If FadeIn and FadeOut Functions are present and they both have a zero duration or if both fade functions are null, remove the transformation
+				// If the number of Fade Regions is invalid
+				// Or if a Fade Region is present and both fades have a zero duration or if both fades are null, remove the transformation
 				SoundWave->Transformations.Remove(TrimFadeTransformation);
 			}
 		}
@@ -318,14 +321,16 @@ TObjectPtr<UWaveformTransformationTrimFade> FWaveformEditor::GetOrAddTrimFadeTra
 			SoundWave->Transformations.Add(TrimFadeTransformation);
 		}
 
+		check(TrimFadeTransformation->FadeTransformation);
+		check(TrimFadeTransformation->FadeTransformation->GetFadeRegionsNum() == 1);
 		// Add default fade curves
-		TrimFadeTransformation->FadeFunctions.FadeIn = NewObject<UFadeFunction>(TrimFadeTransformation.Get(), FadeInType.Get(), NAME_None, RF_Transactional);
-		check(TrimFadeTransformation->FadeFunctions.FadeIn);
-		TrimFadeTransformation->FadeFunctions.FadeOut = NewObject<UFadeFunction>(TrimFadeTransformation.Get(), FadeOutType.Get(), NAME_None, RF_Transactional);
-		check(TrimFadeTransformation->FadeFunctions.FadeOut);
+		TrimFadeTransformation->GetMutableFadeRegion().FadeIn = NewObject<UTransformationFadeFunction>(TrimFadeTransformation->FadeTransformation.Get(), FadeInType.Get(), NAME_None, RF_Transactional);
+		check(TrimFadeTransformation->GetFadeRegion().FadeIn);
+		TrimFadeTransformation->GetMutableFadeRegion().FadeOut = NewObject<UTransformationFadeFunction>(TrimFadeTransformation->FadeTransformation.Get(), FadeOutType.Get(), NAME_None, RF_Transactional);
+		check(TrimFadeTransformation->GetFadeRegion().FadeOut);
 
-		TrimFadeTransformation->FadeFunctions.FadeIn->Duration = 0.f;
-		TrimFadeTransformation->FadeFunctions.FadeOut->Duration = 0.f;
+		TrimFadeTransformation->GetFadeRegion().FadeIn->Duration = 0.f;
+		TrimFadeTransformation->GetFadeRegion().FadeOut->Duration = 0.f;
 
 		TransformationChainConfig = SoundWave->UpdateTransformations();
 	}
@@ -470,15 +475,19 @@ void FWaveformEditor::ToggleFadeIn()
 
 	TObjectPtr<UWaveformTransformationTrimFade> TrimFadeTransformation = GetOrAddTrimFadeTransformation();
 	check(TrimFadeTransformation);
+	check(TrimFadeTransformation->FadeTransformation);
+	check(TrimFadeTransformation->FadeTransformation->GetFadeRegionsNum() == 1);
 
-	if (TrimFadeTransformation->FadeFunctions.FadeIn)
+	const FTransformationFadeFunctionData& ConstFadeRegion = TrimFadeTransformation->GetFadeRegion();
+
+	if (ConstFadeRegion.FadeIn)
 	{
 		GEditor->BeginTransaction(TEXT("PropertyEditor"), LOCTEXT("SetNumericPropertyTransactionToggleFadeIn", "ToggleFadeIn"), nullptr);
 		SoundWave->Modify();
 		TrimFadeTransformation->Modify();
-		TrimFadeTransformation->FadeFunctions.FadeIn->Modify();
+		ConstFadeRegion.FadeIn->Modify();
 		
-		if (TrimFadeTransformation->FadeFunctions.FadeIn->Duration == 0.0f)
+		if (ConstFadeRegion.FadeIn->Duration == 0.0f)
 		{
 			check(DefaultFadeInAmount >= 0.0f);
 
@@ -488,20 +497,24 @@ void FWaveformEditor::ToggleFadeIn()
 			}
 
 			// Ensures the selected subtype (linear, exponential, logarithmic, or sigmoid)
-			if (!TrimFadeTransformation->FadeFunctions.FadeIn.IsA(FadeInType.Get()))
+			if (!ConstFadeRegion.FadeIn.IsA(FadeInType.Get()))
 			{
-				TrimFadeTransformation->FadeFunctions.FadeIn = NewObject<UFadeFunction>(TrimFadeTransformation.Get(), FadeInType.Get(), NAME_None, RF_Transactional);
+				FTransformationFadeFunctionData& MutableFadeRegion = TrimFadeTransformation->GetMutableFadeRegion();
+				MutableFadeRegion.FadeIn = NewObject<UTransformationFadeFunction>(TrimFadeTransformation->FadeTransformation.Get(), FadeInType.Get(), NAME_None, RF_Transactional);
+				MutableFadeRegion.FadeIn->Modify();
+
+				ensure(&ConstFadeRegion == &MutableFadeRegion);
 			}
 
-			TrimFadeTransformation->FadeFunctions.FadeIn->Duration = CachedFadeInAmount;
+			ConstFadeRegion.FadeIn->Duration = CachedFadeInAmount;
 		}
 		else //if fade duration is not 0, assume there is a fade applied and toggle fade "off" by setting fade time to 0
 		{
 			// No need to check for fade type since we are toggling it "off". Type check is done during toggle "on"
-			CachedFadeInAmount = TrimFadeTransformation->FadeFunctions.FadeIn->Duration;
+			CachedFadeInAmount = ConstFadeRegion.FadeIn->Duration;
 
 			//No need to set FadeCurve when toggling off. This will preserve user curve data if they later manually change the StartFadeTime.
-			TrimFadeTransformation->FadeFunctions.FadeIn->Duration = 0.0f;
+			ConstFadeRegion.FadeIn->Duration = 0.0f;
 		}
 		
 		NotifyPostTransformationChange(EPropertyChangeType::ValueSet);
@@ -521,15 +534,19 @@ void FWaveformEditor::ToggleFadeOut()
 
 	TObjectPtr<UWaveformTransformationTrimFade> TrimFadeTransformation = GetOrAddTrimFadeTransformation();
 	check(TrimFadeTransformation);
+	check(TrimFadeTransformation->FadeTransformation);
+	check(TrimFadeTransformation->FadeTransformation->GetFadeRegionsNum() == 1);
 
-	if (TrimFadeTransformation->FadeFunctions.FadeOut)
+	const FTransformationFadeFunctionData& ConstFadeRegion = TrimFadeTransformation->GetFadeRegion();
+
+	if (ConstFadeRegion.FadeOut)
 	{
 		GEditor->BeginTransaction(TEXT("PropertyEditor"), LOCTEXT("SetNumericPropertyTransactionToggleFadeOut", "ToggleFadeOut"), nullptr);
 		SoundWave->Modify();
 		TrimFadeTransformation->Modify();
-		TrimFadeTransformation->FadeFunctions.FadeOut->Modify();
+		ConstFadeRegion.FadeOut->Modify();
 		
-		if (TrimFadeTransformation->FadeFunctions.FadeOut->Duration == 0.0f)
+		if (ConstFadeRegion.FadeOut->Duration == 0.0f)
 		{
 			check(DefaultFadeOutAmount >= 0.0f);
 
@@ -539,20 +556,24 @@ void FWaveformEditor::ToggleFadeOut()
 			}
 
 			// Ensures the selected subtype (linear, exponential, logarithmic, or sigmoid)
-			if (!TrimFadeTransformation->FadeFunctions.FadeOut.IsA(FadeOutType.Get()))
+			if (!ConstFadeRegion.FadeOut.IsA(FadeOutType.Get()))
 			{
-				TrimFadeTransformation->FadeFunctions.FadeOut = NewObject<UFadeFunction>(TrimFadeTransformation.Get(), FadeOutType.Get(), NAME_None, RF_Transactional);
+				FTransformationFadeFunctionData& MutableFadeRegion = TrimFadeTransformation->GetMutableFadeRegion();
+				MutableFadeRegion.FadeOut = NewObject<UTransformationFadeFunction>(TrimFadeTransformation->FadeTransformation.Get(), FadeOutType.Get(), NAME_None, RF_Transactional);
+				MutableFadeRegion.FadeOut->Modify();
+
+				ensure(&ConstFadeRegion == &MutableFadeRegion);
 			}
 
-			TrimFadeTransformation->FadeFunctions.FadeOut->Duration = CachedFadeOutAmount;
+			ConstFadeRegion.FadeOut->Duration = CachedFadeOutAmount;
 		}
 		else //if fade duration is not 0, assume there is a fade applied and toggle fade "off" by setting fade time to 0
 		{
 			// No need to check for fade type since we are toggling it "off". Type check is done during toggle "on"
-			CachedFadeOutAmount = TrimFadeTransformation->FadeFunctions.FadeOut->Duration;
+			CachedFadeOutAmount = ConstFadeRegion.FadeOut->Duration;
 
 			//No need to set FadeCurve when toggling off. This will preserve user curve data if they later manually change the StartFadeTime.
-			TrimFadeTransformation->FadeFunctions.FadeOut->Duration = 0.0f;
+			ConstFadeRegion.FadeOut->Duration = 0.0f;
 		}
 
 		NotifyPostTransformationChange(EPropertyChangeType::ValueSet);
@@ -1500,37 +1521,37 @@ bool FWaveformEditor::BindCommands()
 		Commands.FadeInLinear,
 		FExecuteAction::CreateLambda([this] 
 			{ 
-				FadeInType = UWaveformTransformationTrimFade::FadeModeToFadeFunctionMap[EWaveEditorFadeMode::Linear];
+				FadeInType = UWaveformTransformationFade::FadeModeToFadeFunctionMap[EWaveEditorTransformationFadeMode::Linear];
 			}),
 		FCanExecuteAction::CreateLambda([this] { return true; }),
-		FIsActionChecked::CreateLambda([this] { return FadeInType == UWaveformTransformationTrimFade::FadeModeToFadeFunctionMap[EWaveEditorFadeMode::Linear]; }));
+		FIsActionChecked::CreateLambda([this] { return FadeInType == UWaveformTransformationFade::FadeModeToFadeFunctionMap[EWaveEditorTransformationFadeMode::Linear]; }));
 
 	ToolkitCommands->MapAction(
 		Commands.FadeInExponential,
 		FExecuteAction::CreateLambda([this] 
 			{ 
-				FadeInType = UWaveformTransformationTrimFade::FadeModeToFadeFunctionMap[EWaveEditorFadeMode::Exponential];
+				FadeInType = UWaveformTransformationFade::FadeModeToFadeFunctionMap[EWaveEditorTransformationFadeMode::Exponential];
 			}),
 		FCanExecuteAction::CreateLambda([this] { return true; }),
-		FIsActionChecked::CreateLambda([this] { return FadeInType == UWaveformTransformationTrimFade::FadeModeToFadeFunctionMap[EWaveEditorFadeMode::Exponential]; }));
+		FIsActionChecked::CreateLambda([this] { return FadeInType == UWaveformTransformationFade::FadeModeToFadeFunctionMap[EWaveEditorTransformationFadeMode::Exponential]; }));
 
 	ToolkitCommands->MapAction(
 		Commands.FadeInLogarithmic,
 		FExecuteAction::CreateLambda([this] 
 			{ 
-				FadeInType = UWaveformTransformationTrimFade::FadeModeToFadeFunctionMap[EWaveEditorFadeMode::Logarithmic];
+				FadeInType = UWaveformTransformationFade::FadeModeToFadeFunctionMap[EWaveEditorTransformationFadeMode::Logarithmic];
 			}),
 		FCanExecuteAction::CreateLambda([this] { return true; }),
-		FIsActionChecked::CreateLambda([this] { return FadeInType == UWaveformTransformationTrimFade::FadeModeToFadeFunctionMap[EWaveEditorFadeMode::Logarithmic]; }));
+		FIsActionChecked::CreateLambda([this] { return FadeInType == UWaveformTransformationFade::FadeModeToFadeFunctionMap[EWaveEditorTransformationFadeMode::Logarithmic]; }));
 
 	ToolkitCommands->MapAction(
 		Commands.FadeInSigmoid,
 		FExecuteAction::CreateLambda([this] 
 			{ 
-				FadeInType = UWaveformTransformationTrimFade::FadeModeToFadeFunctionMap[EWaveEditorFadeMode::Sigmoid];
+				FadeInType = UWaveformTransformationFade::FadeModeToFadeFunctionMap[EWaveEditorTransformationFadeMode::Sigmoid];
 			}),
 		FCanExecuteAction::CreateLambda([this] { return true; }),
-		FIsActionChecked::CreateLambda([this] { return FadeInType == UWaveformTransformationTrimFade::FadeModeToFadeFunctionMap[EWaveEditorFadeMode::Sigmoid]; }));
+		FIsActionChecked::CreateLambda([this] { return FadeInType == UWaveformTransformationFade::FadeModeToFadeFunctionMap[EWaveEditorTransformationFadeMode::Sigmoid]; }));
 
 	ToolkitCommands->MapAction(
 		Commands.ToggleFadeOut,
@@ -1541,37 +1562,37 @@ bool FWaveformEditor::BindCommands()
 		Commands.FadeOutLinear,
 		FExecuteAction::CreateLambda([this] 
 			{ 
-				FadeOutType = UWaveformTransformationTrimFade::FadeModeToFadeFunctionMap[EWaveEditorFadeMode::Linear];
+				FadeOutType = UWaveformTransformationFade::FadeModeToFadeFunctionMap[EWaveEditorTransformationFadeMode::Linear];
 			}),
 		FCanExecuteAction::CreateLambda([this] {return true; }),
-		FIsActionChecked::CreateLambda([this] {return FadeOutType == UWaveformTransformationTrimFade::FadeModeToFadeFunctionMap[EWaveEditorFadeMode::Linear]; }));
+		FIsActionChecked::CreateLambda([this] {return FadeOutType == UWaveformTransformationFade::FadeModeToFadeFunctionMap[EWaveEditorTransformationFadeMode::Linear]; }));
 
 	ToolkitCommands->MapAction(
 		Commands.FadeOutExponential,
 		FExecuteAction::CreateLambda([this] 
 			{ 
-				FadeOutType = UWaveformTransformationTrimFade::FadeModeToFadeFunctionMap[EWaveEditorFadeMode::Exponential];
+				FadeOutType = UWaveformTransformationFade::FadeModeToFadeFunctionMap[EWaveEditorTransformationFadeMode::Exponential];
 			}),
 		FCanExecuteAction::CreateLambda([this] {return true; }),
-		FIsActionChecked::CreateLambda([this] {return FadeOutType == UWaveformTransformationTrimFade::FadeModeToFadeFunctionMap[EWaveEditorFadeMode::Exponential]; }));
+		FIsActionChecked::CreateLambda([this] {return FadeOutType == UWaveformTransformationFade::FadeModeToFadeFunctionMap[EWaveEditorTransformationFadeMode::Exponential]; }));
 
 	ToolkitCommands->MapAction(
 		Commands.FadeOutLogarithmic,
 		FExecuteAction::CreateLambda([this] 
 			{ 
-				FadeOutType = UWaveformTransformationTrimFade::FadeModeToFadeFunctionMap[EWaveEditorFadeMode::Logarithmic];
+				FadeOutType = UWaveformTransformationFade::FadeModeToFadeFunctionMap[EWaveEditorTransformationFadeMode::Logarithmic];
 			}),
 		FCanExecuteAction::CreateLambda([this] {return true; }),
-		FIsActionChecked::CreateLambda([this] {return FadeOutType == UWaveformTransformationTrimFade::FadeModeToFadeFunctionMap[EWaveEditorFadeMode::Logarithmic]; }));
+		FIsActionChecked::CreateLambda([this] {return FadeOutType == UWaveformTransformationFade::FadeModeToFadeFunctionMap[EWaveEditorTransformationFadeMode::Logarithmic]; }));
 
 	ToolkitCommands->MapAction(
 		Commands.FadeOutSigmoid,
 		FExecuteAction::CreateLambda([this] 
 			{ 
-				FadeOutType = UWaveformTransformationTrimFade::FadeModeToFadeFunctionMap[EWaveEditorFadeMode::Sigmoid];
+				FadeOutType = UWaveformTransformationFade::FadeModeToFadeFunctionMap[EWaveEditorTransformationFadeMode::Sigmoid];
 			}),
 		FCanExecuteAction::CreateLambda([this] {return true; }),
-		FIsActionChecked::CreateLambda([this] {return FadeOutType == UWaveformTransformationTrimFade::FadeModeToFadeFunctionMap[EWaveEditorFadeMode::Sigmoid]; }));
+		FIsActionChecked::CreateLambda([this] {return FadeOutType == UWaveformTransformationFade::FadeModeToFadeFunctionMap[EWaveEditorTransformationFadeMode::Sigmoid]; }));
 
 	ToolkitCommands->MapAction(
 		Commands.LeftBoundsIncrease,
diff --git a/Engine/Plugins/Editor/WaveformEditor/Source/WaveformEditor/Public/WaveformEditor.h b/Engine/Plugins/Editor/WaveformEditor/Source/WaveformEditor/Public/WaveformEditor.h
--- a/Engine/Plugins/Editor/WaveformEditor/Source/WaveformEditor/Public/WaveformEditor.h
+++ b/Engine/Plugins/Editor/WaveformEditor/Source/WaveformEditor/Public/WaveformEditor.h
@@ -193,11 +193,11 @@ private:
 
 	static constexpr float DefaultFadeInAmount = 0.5f;
 	float CachedFadeInAmount = 0.f;
-	TSubclassOf<UFadeFunction> FadeInType = UWaveformTransformationTrimFade::FadeModeToFadeFunctionMap[EWaveEditorFadeMode::Exponential];
+	TSubclassOf<UTransformationFadeFunction> FadeInType = UWaveformTransformationFade::FadeModeToFadeFunctionMap[EWaveEditorTransformationFadeMode::Exponential];
 
 	static constexpr float DefaultFadeOutAmount = 0.5f;
 	float CachedFadeOutAmount = 0.f;
-	TSubclassOf<UFadeFunction> FadeOutType = UWaveformTransformationTrimFade::FadeModeToFadeFunctionMap[EWaveEditorFadeMode::Exponential];
+	TSubclassOf<UTransformationFadeFunction> FadeOutType = UWaveformTransformationFade::FadeModeToFadeFunctionMap[EWaveEditorTransformationFadeMode::Exponential];
 
 	// Tracking bool to ensure we detect wave cue changes and regenerate transformations
 	// Checking the event is bound may have issues if other systems need to subscribe for wave cue array changes
diff --git a/Engine/Plugins/Editor/WaveformEditor/Source/WaveformTransformations/Private/WaveformTransformationFade.cpp b/Engine/Plugins/Editor/WaveformEditor/Source/WaveformTransformations/Private/WaveformTransformationFade.cpp
new file mode 100644
--- /dev/null
+++ b/Engine/Plugins/Editor/WaveformEditor/Source/WaveformTransformations/Private/WaveformTransformationFade.cpp
@@ -0,0 +1,445 @@
+// Copyright Epic Games, Inc. All Rights Reserved.
+
+#include "WaveformTransformationFade.h"
+
+#include "AudioCompressionSettingsUtils.h"
+#include "Editor.h"
+#include "WaveformTransformationLog.h"
+
+#include UE_INLINE_GENERATED_CPP_BY_NAME(WaveformTransformationFade)
+
+namespace WaveformTransformationFadePrivate
+{
+	// Makes the two functions for SCurve meet at (0.5, 0.5)
+	const double XMultiplier = FMath::Sqrt(10.0) / 5.0;
+
+	constexpr float Linear = 1.f;
+
+	static void ApplyFadeIn(Audio::FAlignedFloatBuffer& InputAudio, const int64 AudioSizeInSamples, const FTransformationFadeFunctionData& FadeFunctions, const int32 NumChannels, const float SampleRate, const int64 StartFrameOffset)
+	{
+		if (FadeFunctions.FadeIn == nullptr)
+		{
+			return;
+		}
+
+		check(NumChannels > 0);
+		check(SampleRate > 0.f);
+
+		// Assert against buffer overflows
+		ensure(StartFrameOffset >= 0);
+		ensure(StartFrameOffset == 0 || (INT64_MAX / StartFrameOffset) >= NumChannels);
+		const int64 StartSampleOffset = static_cast<int64>(StartFrameOffset * NumChannels);
+		ensure(INT64_MAX - AudioSizeInSamples >= StartSampleOffset);
+
+		const int64 NumSamples = (AudioSizeInSamples + StartSampleOffset);
+		const int64 NumFrames = NumSamples / NumChannels;
+
+		const int64 StartFrame = FadeFunctions.FadeIn->FrameOffset - StartFrameOffset;
+
+		if (StartFrame < 0)
+		{
+			return;
+		}
+
+		const float FadeDuration = FadeFunctions.FadeIn->Duration;
+		const int64 NumFramesAfterOffset = NumFrames - StartFrame;
+		const int64 FadeNumFrames = FMath::Min(FMath::FloorToInt64(FadeDuration * SampleRate), NumFramesAfterOffset);
+
+		if (NumFrames <= 1 || FadeDuration < SMALL_NUMBER || FadeNumFrames == 0 || (StartFrame + FadeNumFrames) * NumChannels > InputAudio.Num())
+		{
+			return;
+		}
+
+		float* InputPtr = InputAudio.GetData() + StartFrame * NumChannels;
+
+		for (int64 FrameIndex = 0; FrameIndex < FadeNumFrames; ++FrameIndex)
+		{
+			const float FadeFraction = static_cast<float>(FrameIndex) / static_cast<float>(FadeNumFrames);
+			check(FadeFraction <= 1.0f);
+			check(FadeFraction >= 0.0f);
+
+			const float EnvValue = FadeFunctions.FadeIn->GetFadeInCurveValue(FadeFraction);
+
+			for (int64 ChannelIt = 0; ChannelIt < NumChannels; ++ChannelIt)
+			{
+				*(InputPtr + ChannelIt) *= EnvValue;
+			}
+
+			InputPtr += NumChannels;
+		}
+	}
+
+	static void ApplyFadeOut(Audio::FAlignedFloatBuffer& InputAudio, const int64 AudioSizeInSamples, const FTransformationFadeFunctionData& FadeFunctions, const int32 NumChannels, const float SampleRate, const int64 EndFrameOffset)
+	{
+		if (FadeFunctions.FadeOut == nullptr)
+		{
+			return;
+		}
+
+		check(NumChannels > 0);
+		check(SampleRate > 0.f);
+
+		// Assert against buffer overflows
+		ensure(EndFrameOffset == 0 || (INT64_MAX / EndFrameOffset) >= NumChannels);
+		const int64 EndSampleOffset = static_cast<int64>(EndFrameOffset * NumChannels);
+		ensure(INT64_MAX - AudioSizeInSamples >= EndSampleOffset);
+
+		const int64 NumSamples = (AudioSizeInSamples + EndSampleOffset);
+		const int64 NumFramesWithEndFrameOffset = NumSamples / NumChannels;
+
+		const int64 FrameOffset = FadeFunctions.FadeOut->FrameOffset;
+
+		if (FrameOffset < 0)
+		{
+			return;
+		}
+
+		const float FadeDuration = FadeFunctions.FadeOut->Duration;
+		const int64 NumFramesAfterOffset = NumFramesWithEndFrameOffset - FrameOffset;
+		const int64 FadeNumFrames = FMath::Min(FMath::FloorToInt64(FadeDuration * SampleRate), NumFramesAfterOffset);
+
+		if (NumFramesAfterOffset <= 1 || FadeDuration < SMALL_NUMBER || FadeNumFrames == 0)
+		{
+			return;
+		}
+
+		const int64 FadeNumFramesFromEnd = FMath::Min(FadeNumFrames + FrameOffset, NumFramesWithEndFrameOffset);
+		const int64 StartSampleIndex = NumSamples - (static_cast<int64>(FadeNumFramesFromEnd) * NumChannels);
+		check(StartSampleIndex >= 0);
+		check(AudioSizeInSamples <= InputAudio.Num());
+
+		if (StartSampleIndex + (static_cast<int64>(FadeNumFrames) * NumChannels) > AudioSizeInSamples)
+		{
+			return;
+		}
+
+		float* InputPtr = &InputAudio[StartSampleIndex];
+
+		for (int64 FrameIndex = 0; FrameIndex < FadeNumFrames; ++FrameIndex)
+		{
+			const float FadeFraction = static_cast<float>(FrameIndex) / static_cast<float>(FadeNumFrames);
+			check(FadeFraction <= 1.0f);
+			check(FadeFraction >= 0.0f);
+
+			const float EnvValue = FadeFunctions.FadeOut->GetFadeOutCurveValue(FadeFraction);
+
+			for (int64 ChannelIt = 0; ChannelIt < NumChannels; ++ChannelIt)
+			{
+				*(InputPtr + ChannelIt) *= EnvValue;
+			}
+
+			InputPtr += NumChannels;
+		}
+	}
+}
+
+const TMap<EWaveEditorTransformationFadeMode, TSubclassOf<UTransformationFadeFunction>> UWaveformTransformationFade::FadeModeToFadeFunctionMap =
+{
+	{EWaveEditorTransformationFadeMode::Linear, UTransformationFadeFunctionLinear::StaticClass()},
+	{EWaveEditorTransformationFadeMode::Exponential, UTransformationFadeCurveFunctionExponential::StaticClass()},
+	{EWaveEditorTransformationFadeMode::Logarithmic, UTransformationFadeCurveFunctionLogarithmic::StaticClass()},
+	{EWaveEditorTransformationFadeMode::Sigmoid, UTransformationFadeCurveFunctionSigmoid::StaticClass()}
+};
+
+FWaveTransformationFade::FWaveTransformationFade(const TArray<FTransformationFadeFunctionData>& InFadeRegions, const int64 InStartFrameOffset, const int64 InEndFrameOffset)
+	: FadeRegions(InFadeRegions)
+	, StartFrameOffset(InStartFrameOffset)
+	, EndFrameOffset(InEndFrameOffset) {}
+
+void FWaveTransformationFade::ProcessAudio(Audio::FWaveformTransformationWaveInfo& InOutWaveInfo) const
+{
+	check(InOutWaveInfo.SampleRate > 0.f && InOutWaveInfo.Audio != nullptr);
+	check(InOutWaveInfo.NumChannels > 0);
+
+	Audio::FAlignedFloatBuffer& InputAudio = *InOutWaveInfo.Audio;
+
+	// If InputAudio is smaller than the size of a frame, do not process the audio.
+	const int64 InputAudioNum = static_cast<int64>(InputAudio.Num());
+	check(InputAudioNum >= InOutWaveInfo.NumChannels);
+
+	InOutWaveInfo.NumEditedSamples = InputAudioNum;
+
+	for (const FTransformationFadeFunctionData& Region : FadeRegions)
+	{
+		const bool bProcessFadeIn = Region.FadeIn && Region.FadeIn->Duration > 0.f;
+		const bool bProcessFadeOut = Region.FadeOut && Region.FadeOut->Duration > 0.f;
+
+		if (bProcessFadeIn)
+		{
+			WaveformTransformationFadePrivate::ApplyFadeIn(InputAudio, InputAudioNum, Region, InOutWaveInfo.NumChannels, InOutWaveInfo.SampleRate, StartFrameOffset);
+		}
+
+		if (bProcessFadeOut)
+		{
+			WaveformTransformationFadePrivate::ApplyFadeOut(InputAudio, InputAudioNum, Region, InOutWaveInfo.NumChannels, InOutWaveInfo.SampleRate, EndFrameOffset);
+		}
+	}
+}
+
+UWaveformTransformationFade::UWaveformTransformationFade(const FObjectInitializer& ObjectInitializer) : Super(ObjectInitializer)
+{
+	FTransformationFadeFunctionData FadeRegion;
+	
+	FadeRegion.FadeIn = CreateDefaultSubobject<UTransformationFadeCurveFunctionExponential>(TEXT("FadeInFunction"));
+	FadeRegion.FadeIn->Duration = 0.f;
+	FadeRegion.FadeOut = CreateDefaultSubobject<UTransformationFadeCurveFunctionExponential>(TEXT("FadeOutFunction"));
+	FadeRegion.FadeOut->Duration = 0.f;
+	
+	FadeRegions.Add(FadeRegion);
+}
+
+#if WITH_EDITOR
+void UWaveformTransformationFade::PostEditChangeProperty(FPropertyChangedEvent& PropertyChangedEvent)
+{
+	if (AvailableWaveformDuration >= 0.f)
+	{
+		// if AvailableWaveformDuration >= 0, we know SampleRate has been initialized
+		check(SampleRate > 0.f);
+		const float SoundDuration = OriginalNumFrames / SampleRate;
+
+		// Clamp frame offset to be within the original waveform 
+		const int64 MaxFrameOffset = OriginalNumFrames;
+
+		for (FTransformationFadeFunctionData& FadeRegion : FadeRegions)
+		{			
+			if (FadeRegion.FadeIn)
+			{
+				FadeRegion.FadeIn->FrameOffset = FMath::Clamp(FadeRegion.FadeIn->FrameOffset, 0, MaxFrameOffset);
+
+				const float MaxFadeInDuration = SoundDuration - (FadeRegion.FadeIn->FrameOffset / SampleRate);
+				FadeRegion.FadeIn->Duration = FMath::Clamp(FadeRegion.FadeIn->Duration, 0.f, MaxFadeInDuration);
+			}
+			else
+			{
+				FadeRegion.FadeIn = NewObject<UTransformationFadeCurveFunctionExponential>(this, UTransformationFadeCurveFunctionExponential::StaticClass(), NAME_None, RF_Transactional);
+				FadeRegion.FadeIn->Duration = 0.f;
+			}
+
+			if (FadeRegion.FadeOut)
+			{
+				FadeRegion.FadeOut->FrameOffset = FMath::Clamp(FadeRegion.FadeOut->FrameOffset, 0, MaxFrameOffset);
+
+				const float MaxFadeOutDuration = SoundDuration - (FadeRegion.FadeOut->FrameOffset / SampleRate);;
+
+				if (FadeRegion.bLinkDurations)
+				{
+					FadeRegion.FadeOut->Duration = FMath::Clamp(FadeRegion.FadeIn->Duration, 0, MaxFadeOutDuration);
+				}
+				else
+				{
+					FadeRegion.FadeOut->Duration = FMath::Clamp(FadeRegion.FadeOut->Duration, 0.f, MaxFadeOutDuration);
+				}
+			}
+			else
+			{
+				FadeRegion.FadeOut = NewObject<UTransformationFadeCurveFunctionExponential>(this, UTransformationFadeCurveFunctionExponential::StaticClass(), NAME_None, RF_Transactional);
+				FadeRegion.FadeOut->Duration = 0.f;
+			}
+		}
+	}
+
+	// Clamp values first, before triggering the transformation update in UWaveformTransformationBase::PostEditChangeProperty(PropertyChangedEvent)
+	Super::PostEditChangeProperty(PropertyChangedEvent);
+}
+
+FString UWaveformTransformationFade::GetTransformationHash() const
+{
+	using FPCU = FPlatformCompressionUtilities;
+	FString Hash;
+
+	float FadeInDuration = 0.f;
+	int64 FadeInFrameOffset = 0;
+	float FadeInCurve = WaveformTransformationFadePrivate::Linear;
+
+	float FadeOutDuration = 0.f;
+	int64 FadeOutFrameOffset = 0;
+	float FadeOutCurve = WaveformTransformationFadePrivate::Linear;
+
+	if (FadeRegions.Num() > 0)
+	{
+		FPCU::AppendHash(Hash, TEXT("FRC"), FadeRegions.Num());
+	}
+
+	for (const FTransformationFadeFunctionData& FadeRegion : FadeRegions)
+	{
+		FPCU::AppendHash(Hash, TEXT("FRLD"), FadeRegion.bLinkDurations);
+
+		if (FadeRegion.FadeIn)
+		{
+			FadeInDuration = FadeRegion.FadeIn->Duration;
+			FadeInFrameOffset = FadeRegion.FadeIn->FrameOffset;
+			if (FadeRegion.FadeIn.IsA<UTransformationFadeCurveFunction>())
+			{
+				TObjectPtr<UTransformationFadeCurveFunction> CurveFunction = CastChecked<UTransformationFadeCurveFunction>(FadeRegion.FadeIn);
+				FadeInCurve = CurveFunction->GetFadeCurve();
+			}
+			else
+			{
+				FadeInCurve = WaveformTransformationFadePrivate::Linear;
+			}
+
+			FPCU::AppendHash(Hash, TEXT("FIC"), FadeInCurve);
+			FPCU::AppendHash(Hash, TEXT("FID"), FadeInDuration);
+			FPCU::AppendHash(Hash, TEXT("FIFO"), FadeInFrameOffset);
+		}
+
+		if (FadeRegion.FadeOut)
+		{
+			FadeOutDuration = FadeRegion.FadeOut->Duration;
+			FadeOutFrameOffset = FadeRegion.FadeOut->FrameOffset;
+			if (FadeRegion.FadeOut.IsA<UTransformationFadeCurveFunction>())
+			{
+				TObjectPtr<UTransformationFadeCurveFunction> CurveFunction = CastChecked<UTransformationFadeCurveFunction>(FadeRegion.FadeOut);
+				FadeOutCurve = CurveFunction->GetFadeCurve();
+			}
+			else
+			{
+				FadeOutCurve = WaveformTransformationFadePrivate::Linear;
+			}
+
+			FPCU::AppendHash(Hash, TEXT("FOC"), FadeOutCurve);
+			FPCU::AppendHash(Hash, TEXT("FOD"), FadeOutDuration);
+			FPCU::AppendHash(Hash, TEXT("FOFO"), FadeOutFrameOffset);
+		}
+	}
+	
+	return Hash;
+}
+#endif // WITH_EDITOR
+
+Audio::FTransformationPtr UWaveformTransformationFade::CreateTransformation() const
+{
+	return MakeUnique<FWaveTransformationFade>(FadeRegions, StartFrameOffset, EndFrameOffset);
+}
+
+void UWaveformTransformationFade::UpdateConfiguration(FWaveTransformUObjectConfiguration& InOutConfiguration)
+{
+	UpdateDurationProperties(InOutConfiguration.EndTime - InOutConfiguration.StartTime);
+	SampleRate = InOutConfiguration.SampleRate;
+	StartFrameOffset = InOutConfiguration.GetStartFrameOffset();
+	EndFrameOffset = InOutConfiguration.GetEndFrameOffset();
+	OriginalNumFrames = InOutConfiguration.GetOriginalNumFrames();
+}
+
+void UWaveformTransformationFade::LimitFadeRegionsToOne()
+{
+	if (FadeRegions.Num() > 1)
+	{
+		FadeRegions.SetNum(1, EAllowShrinking::Yes);
+		UE_LOG(LogWaveformTransformation, Warning, TEXT("Fade Transformation had more than one fade region. REMOVING excess fade regions!"));
+	}
+}
+
+void UWaveformTransformationFade::UpdateDurationProperties(const float InAvailableDuration)
+{
+	check(InAvailableDuration >= 0.f);
+	AvailableWaveformDuration = InAvailableDuration;
+}
+
+const double UTransformationFadeFunctionLinear::GetFadeInCurveValue(const double FadeFraction) const
+{
+	check(FadeFraction <= 1.f);
+	check(FadeFraction >= 0.f);
+
+	return FadeFraction;
+}
+
+const double UTransformationFadeFunctionLinear::GetFadeOutCurveValue(const double FadeFraction) const
+{
+	check(FadeFraction <= 1.f);
+	check(FadeFraction >= 0.f);
+
+	return 1.0 - FadeFraction;
+}
+
+const double UTransformationFadeCurveFunctionExponential::GetFadeInCurveValue(const double FadeFraction) const
+{
+	check(FadeFraction <= 1.f);
+	check(FadeFraction >= 0.f);
+
+	return FMath::Pow(FadeFraction, FadeCurve);
+}
+
+const double UTransformationFadeCurveFunctionExponential::GetFadeOutCurveValue(const double FadeFraction) const
+{
+	// Pow functions return NaN for negative bases with non-integer exponents
+	check(FadeFraction <= 1.f);
+	check(FadeFraction >= 0.f);
+
+	return FMath::Pow(-FadeFraction + 1.0, FadeCurve);
+}
+
+const double UTransformationFadeCurveFunctionLogarithmic::GetFadeInCurveValue(const double FadeFraction) const
+{
+	check(FadeFraction <= 1.f);
+	check(FadeFraction >= 0.f);
+
+	return FMath::Pow(FadeFraction, FadeCurve);
+}
+
+const double UTransformationFadeCurveFunctionLogarithmic::GetFadeOutCurveValue(const double FadeFraction) const
+{
+	// Pow functions return NaN for negative bases with non-integer exponents
+	check(FadeFraction <= 1.f);
+	check(FadeFraction >= 0.f);
+
+	return FMath::Pow(-FadeFraction + 1.0, FadeCurve);
+}
+
+const double UTransformationFadeCurveFunctionSigmoid::GetFadeInCurveValue(const double FadeFraction) const
+{
+	// LogX will produce NaN if FadeFraction is <= 0 or is 1
+	check(FadeFraction <= 1.f);
+	check(FadeFraction >= 0.f);
+	check(WaveformTransformationFadePrivate::XMultiplier > 0.0);
+
+	if (FadeFraction == 0.f || FadeFraction == 1.f)
+	{
+		return FadeFraction;
+	}
+
+	double CurveValue = 1;
+	const double x = FadeFraction;
+	const double a = static_cast<float>(-FMath::Clamp(SigmoidFadeCurve, 0.1f, 1.0f));
+	const double k = 20 * FMath::Pow(a, 5);
+
+	if (FadeFraction <= 0.5)
+	{
+		CurveValue = -(FMath::Tanh(k * PI * FMath::LogX(10, WaveformTransformationFadePrivate::XMultiplier * x) + ((k * PI) / 2))) / 2 + 0.5;
+	}
+	else
+	{
+		CurveValue = (FMath::Tanh(k * PI * FMath::LogX(10, WaveformTransformationFadePrivate::XMultiplier * (-x + 1)) + ((k * PI) / 2))) / 2 + 0.5;
+	}
+
+	return CurveValue;
+}
+
+const double UTransformationFadeCurveFunctionSigmoid::GetFadeOutCurveValue(const double FadeFraction) const
+{
+	// LogX will produce NaN if FadeFraction is <= 0 or is 1
+	check(FadeFraction <= 1.f);
+	check(FadeFraction >= 0.f);
+	check(WaveformTransformationFadePrivate::XMultiplier > 0.0);
+
+	if (FadeFraction == 0.f || FadeFraction == 1.f)
+	{
+		return 1.0 - FadeFraction;
+	}
+
+	double CurveValue = 1;
+	const double x = FadeFraction;
+	const double a = static_cast<float>(FMath::Clamp(SigmoidFadeCurve, 0.1f, 1.0f));
+	const double k = 20 * FMath::Pow(a, 5);
+
+	if (FadeFraction <= 0.5)
+	{
+		CurveValue = -(FMath::Tanh(k * PI * FMath::LogX(10, WaveformTransformationFadePrivate::XMultiplier * x) + ((k * PI) / 2))) / 2 + 0.5;
+	}
+	else
+	{
+		CurveValue = (FMath::Tanh(k * PI * FMath::LogX(10, WaveformTransformationFadePrivate::XMultiplier * (-x + 1)) + ((k * PI) / 2))) / 2 + 0.5;
+	}
+
+	return CurveValue;
+}
diff --git a/Engine/Plugins/Editor/WaveformEditor/Source/WaveformTransformations/Private/WaveformTransformationFadeCustomization.cpp b/Engine/Plugins/Editor/WaveformEditor/Source/WaveformTransformations/Private/WaveformTransformationFadeCustomization.cpp
new file mode 100644
--- /dev/null
+++ b/Engine/Plugins/Editor/WaveformEditor/Source/WaveformTransformations/Private/WaveformTransformationFadeCustomization.cpp
@@ -0,0 +1,35 @@
+// Copyright Epic Games, Inc. All Rights Reserved.
+
+#include "WaveformTransformationFadeCustomization.h"
+
+#include "DetailWidgetRow.h"
+#include "IDetailChildrenBuilder.h"
+#include "IDetailGroup.h"
+#include "IPropertyUtilities.h"
+#include "PropertyHandle.h"
+#include "WaveformTransformationFade.h"
+
+TSharedRef<IPropertyTypeCustomization> FWaveformTransformationFadeCustomization::MakeInstance()
+{
+	return MakeShareable(new FWaveformTransformationFadeCustomization);
+}
+
+void FWaveformTransformationFadeCustomization::CustomizeHeader(TSharedRef<IPropertyHandle> PropertyHandle,
+	FDetailWidgetRow& HeaderRow, IPropertyTypeCustomizationUtils& CustomizationUtils)
+{
+	constexpr bool bDisplayDefaultPropertyButtons = false;
+
+	HeaderRow.NameContent()[PropertyHandle->CreatePropertyNameWidget()]
+	.ValueContent()[PropertyHandle->CreatePropertyValueWidget(bDisplayDefaultPropertyButtons)];
+}
+
+void FWaveformTransformationFadeCustomization::CustomizeChildren(TSharedRef<IPropertyHandle> PropertyHandle,
+	IDetailChildrenBuilder& ChildBuilder, IPropertyTypeCustomizationUtils& CustomizationUtils)
+{
+	TSharedPtr<IPropertyHandle> FadeRegionHandle = PropertyHandle->GetChildHandle(UWaveformTransformationFade::GetFadeRegionsPropertyName());
+
+	if (FadeRegionHandle && FadeRegionHandle->IsValidHandle())
+	{
+		ChildBuilder.AddProperty(FadeRegionHandle.ToSharedRef());
+	}
+}
\ No newline at end of file
diff --git a/Engine/Plugins/Editor/WaveformEditor/Source/WaveformTransformations/Private/WaveformTransformationFadeCustomization.h b/Engine/Plugins/Editor/WaveformEditor/Source/WaveformTransformations/Private/WaveformTransformationFadeCustomization.h
new file mode 100644
--- /dev/null
+++ b/Engine/Plugins/Editor/WaveformEditor/Source/WaveformTransformations/Private/WaveformTransformationFadeCustomization.h
@@ -0,0 +1,18 @@
+// Copyright Epic Games, Inc. All Rights Reserved.
+
+#pragma once
+#include "IPropertyTypeCustomization.h"
+
+/**
+ * Using IPropertyTypeCustomization to edit the details view of the Fade transformation and remove unneeded information in the waveform editor
+ */
+class FWaveformTransformationFadeCustomization : public IPropertyTypeCustomization
+{
+public:
+	static TSharedRef<IPropertyTypeCustomization> MakeInstance();
+
+	virtual void CustomizeHeader(TSharedRef<IPropertyHandle> PropertyHandle, FDetailWidgetRow& HeaderRow,
+								 IPropertyTypeCustomizationUtils& CustomizationUtils) override;
+	virtual void CustomizeChildren(TSharedRef<IPropertyHandle> PropertyHandle, IDetailChildrenBuilder& ChildBuilder,
+								   IPropertyTypeCustomizationUtils& CustomizationUtils) override;
+};
diff --git a/Engine/Plugins/Editor/WaveformEditor/Source/WaveformTransformations/Private/WaveformTransformationTrimFade.cpp b/Engine/Plugins/Editor/WaveformEditor/Source/WaveformTransformations/Private/WaveformTransformationTrimFade.cpp
--- a/Engine/Plugins/Editor/WaveformEditor/Source/WaveformTransformations/Private/WaveformTransformationTrimFade.cpp
+++ b/Engine/Plugins/Editor/WaveformEditor/Source/WaveformTransformations/Private/WaveformTransformationTrimFade.cpp
@@ -15,7 +15,7 @@ namespace WaveformTrasnformationTrimFadePrivate
 	const double XMultiplier = FMath::Sqrt(10.0) / 5.0;
 
 	constexpr float Linear = 1.f;
-
+	PRAGMA_DISABLE_DEPRECATION_WARNINGS
 	static void SetFadeCurve(TObjectPtr<UFadeFunction> FadeFunction, float FadeCurve, float Sharpness)
 	{
 		if (FadeFunction->IsA<UFadeCurveFunction>())
@@ -36,25 +36,68 @@ namespace WaveformTrasnformationTrimFadePrivate
 			
 		}
 	}
-}
 
-PRAGMA_DISABLE_DEPRECATION_WARNINGS
-const TMap<EWaveEditorFadeMode, float> UWaveformTransformationTrimFade::FadeModeToCurveValueMap =
-{
-	{EWaveEditorFadeMode::Linear, 1.0f},
-	{EWaveEditorFadeMode::Exponential, 3.0f},
-	{EWaveEditorFadeMode::Logarithmic, 0.25f},
-	{EWaveEditorFadeMode::Sigmoid, UWaveformTransformationTrimFade::MinFadeCurve}
-};
-PRAGMA_ENABLE_DEPRECATION_WARNINGS
+	static void MigrateToTransformationFadeFunction(const TObjectPtr<UWaveformTransformationFade>& InFadeTransformation, const TObjectPtr<UFadeFunction>& InFadeFunction, TObjectPtr<UTransformationFadeFunction>& NewTransformationFadeFunction)
+	{
+		check(InFadeTransformation);
+		check(InFadeFunction);
 
-const TMap<EWaveEditorFadeMode, TSubclassOf<UFadeFunction>> UWaveformTransformationTrimFade::FadeModeToFadeFunctionMap =
-{
-	{EWaveEditorFadeMode::Linear, UFadeFunctionLinear::StaticClass()},
-	{EWaveEditorFadeMode::Exponential, UFadeCurveFunctionExponential::StaticClass()},
-	{EWaveEditorFadeMode::Logarithmic, UFadeCurveFunctionLogarithmic::StaticClass()},
-	{EWaveEditorFadeMode::Sigmoid, UFadeCurveFunctionSigmoid::StaticClass()}
-};
+		bool bValidFunctionType = true;
+
+		if (InFadeFunction->IsA<UFadeCurveFunction>())
+		{
+			TObjectPtr<UFadeCurveFunction> CurveFunction = CastChecked<UFadeCurveFunction>(InFadeFunction);
+
+			if (CurveFunction)
+			{
+				if (CurveFunction->IsA<UFadeCurveFunctionExponential>())
+				{
+					NewTransformationFadeFunction = NewObject<UTransformationFadeCurveFunctionExponential>(InFadeTransformation, UTransformationFadeCurveFunctionExponential::StaticClass(), NAME_None, RF_Transactional);
+				}
+				else if (CurveFunction->IsA<UFadeCurveFunctionLogarithmic>())
+				{
+					NewTransformationFadeFunction = NewObject<UTransformationFadeCurveFunctionLogarithmic>(InFadeTransformation, UTransformationFadeCurveFunctionLogarithmic::StaticClass(), NAME_None, RF_Transactional);
+				}
+				else if (CurveFunction->IsA<UFadeCurveFunctionSigmoid>())
+				{
+					NewTransformationFadeFunction = NewObject<UTransformationFadeCurveFunctionSigmoid>(InFadeTransformation, UTransformationFadeCurveFunctionSigmoid::StaticClass(), NAME_None, RF_Transactional);
+				}
+				else
+				{
+					bValidFunctionType = false;
+				}
+
+				if (bValidFunctionType)
+				{
+					TObjectPtr<UTransformationFadeCurveFunction> TransformationCurve = CastChecked<UTransformationFadeCurveFunction>(NewTransformationFadeFunction);
+
+					if (TransformationCurve)
+					{
+						TransformationCurve->SetFadeCurve(CurveFunction->GetFadeCurve());
+					}
+				}
+			}
+		}
+		else if (InFadeFunction->IsA<UFadeFunctionLinear>())
+		{
+			NewTransformationFadeFunction = NewObject<UTransformationFadeFunctionLinear>(InFadeTransformation, UTransformationFadeFunctionLinear::StaticClass(), NAME_None, RF_Transactional);
+		}
+		else
+		{
+			bValidFunctionType = false;
+		}
+
+		if (!bValidFunctionType)
+		{
+			UE_LOG(LogWaveformTransformation, Warning, TEXT("Unknown Curve Function in TrimFade Transformation Migration. Initializing it as exponential"));
+
+			NewTransformationFadeFunction = NewObject<UTransformationFadeCurveFunctionExponential>(InFadeTransformation, UTransformationFadeCurveFunctionExponential::StaticClass(), NAME_None, RF_Transactional);
+		}
+
+		NewTransformationFadeFunction->Duration = InFadeFunction->Duration;
+	}
+	PRAGMA_ENABLE_DEPRECATION_WARNINGS
+}
 
 const float UWaveformTransformationTrimFade::MinFadeCurve = -0.1f;
 const float UWaveformTransformationTrimFade::MaxFadeCurve = 10.f;
@@ -65,6 +108,7 @@ namespace WaveformTransformationTrimFadeNames
 	static FLazyName EndTimeName(GET_MEMBER_NAME_CHECKED(UWaveformTransformationTrimFade, EndTime));
 }
 
+PRAGMA_DISABLE_DEPRECATION_WARNINGS
 static void ApplyFadeIn(Audio::FAlignedFloatBuffer& InputAudio, const FFadeFunctionData& FadeFunctions, const int32 NumChannels, const float SampleRate)
 {
 	if (FadeFunctions.FadeIn == nullptr)
@@ -139,9 +183,10 @@ static void ApplyFadeOut(Audio::FAlignedFloatBuffer& InputAudio, const FFadeFunc
 		InputPtr += NumChannels;
 	}
 }
+PRAGMA_ENABLE_DEPRECATION_WARNINGS
 
-FWaveTransformationTrimFade::FWaveTransformationTrimFade(const FFadeFunctionData& InFadeFunctions, double InStartTime, double InEndTime)
-	: FadeFunctions(InFadeFunctions)
+FWaveTransformationTrimFade::FWaveTransformationTrimFade(TObjectPtr<UWaveformTransformationFade> InFadeTransformation, double InStartTime, double InEndTime)
+	: FadeTransformation(InFadeTransformation)
 	, StartTime(InStartTime)
 	, EndTime(InEndTime){}
 
@@ -151,7 +196,6 @@ void FWaveTransformationTrimFade::ProcessAudio(Audio::FWaveformTransformationWav
 
 	Audio::FAlignedFloatBuffer& InputAudio = *InOutWaveInfo.Audio;
 
-	// If InputAudio is smaller than the size of a frame, do not process the audio.
 	check(InOutWaveInfo.NumChannels > 0);
 
 	int32 InputAudioNum = InputAudio.Num();
@@ -163,6 +207,7 @@ void FWaveTransformationTrimFade::ProcessAudio(Audio::FWaveformTransformationWav
 		UE_LOG(LogWaveformTransformation, Log, TEXT("Invalid number of Samples, number of samples not divisible by the channel count."));
 	}
 
+	// If InputAudio is smaller than the size of a frame, do not process the audio.
 	if (InputAudioNum < InOutWaveInfo.NumChannels)
 	{
 		return;
@@ -199,106 +244,105 @@ void FWaveTransformationTrimFade::ProcessAudio(Audio::FWaveformTransformationWav
 	check(FinalSizeInSamples > 0);
 
 	const bool bApplyTrim = FinalSizeInSamples < InputAudio.Num();
-	const bool bProcessFadeIn = FadeFunctions.FadeIn && FadeFunctions.FadeIn->Duration > 0.f;
-	const bool bProcessFadeOut = FadeFunctions.FadeOut && FadeFunctions.FadeOut->Duration > 0.f;
 
 	if (bApplyTrim)
 	{
 		FMemory::Memmove(InputAudio.GetData(), &InputAudio[StartSample], FinalSizeInSamples * sizeof(float));
 		InputAudio.SetNum(FinalSizeInSamples, EAllowShrinking::Yes);
-		
-		InOutWaveInfo.StartSampleOffset = StartSample;
-	}
-	else
-	{
-		InOutWaveInfo.StartSampleOffset = 0;
 	}
 
 	InOutWaveInfo.NumEditedSamples = InputAudio.Num();
 
-	if (bProcessFadeIn)
+	if (FadeTransformation)
 	{
-		ApplyFadeIn(InputAudio, FadeFunctions, InOutWaveInfo.NumChannels, InOutWaveInfo.SampleRate);
+		FadeTransformation->CreateTransformation()->ProcessAudio(InOutWaveInfo);
 	}
 
-	if (bProcessFadeOut)
+	if (bApplyTrim)
 	{
-		ApplyFadeOut(InputAudio, FadeFunctions, InOutWaveInfo.NumChannels, InOutWaveInfo.SampleRate);
+		// Done after fade transformation to avoid breaking the start offset
+		InOutWaveInfo.StartSampleOffset = StartSample;
 	}
 }
 
 UWaveformTransformationTrimFade::UWaveformTransformationTrimFade(const FObjectInitializer& ObjectInitializer) : Super(ObjectInitializer)
 {
+	FadeTransformation = CreateDefaultSubobject<UWaveformTransformationFade>(TEXT("FadeTransformation"));
+
+	PRAGMA_DISABLE_DEPRECATION_WARNINGS
 	if (FadeFunctions.FadeIn == nullptr)
 	{
-		PRAGMA_DISABLE_DEPRECATION_WARNINGS
-			if (StartFadeTime != 0.f)
+		if (StartFadeTime != 0.f)
+		{
+			if (StartFadeCurve > WaveformTrasnformationTrimFadePrivate::Linear)
 			{
-				if (StartFadeCurve > WaveformTrasnformationTrimFadePrivate::Linear)
-				{
-					FadeFunctions.FadeIn = CreateDefaultSubobject<UFadeCurveFunctionExponential>(TEXT("FadeInFunction"));
-				}
-				else if (StartFadeCurve < 0.f)
-				{
-					FadeFunctions.FadeIn = CreateDefaultSubobject<UFadeCurveFunctionSigmoid>(TEXT("FadeInFunction"));
-				}
-				else if (StartFadeCurve < WaveformTrasnformationTrimFadePrivate::Linear)
-				{
-					FadeFunctions.FadeIn = CreateDefaultSubobject<UFadeCurveFunctionLogarithmic>(TEXT("FadeInFunction"));
-				}
-				else if (StartFadeCurve == WaveformTrasnformationTrimFadePrivate::Linear)
-				{
-					FadeFunctions.FadeIn = CreateDefaultSubobject<UFadeFunctionLinear>(TEXT("FadeInFunction"));
-				}
-
-				FadeFunctions.FadeIn->Duration = StartFadeTime;
-				WaveformTrasnformationTrimFadePrivate::SetFadeCurve(FadeFunctions.FadeIn, StartFadeCurve, StartSCurveSharpness);
-
-				StartFadeTime = 0.f;
-				PRAGMA_ENABLE_DEPRECATION_WARNINGS
+				FadeFunctions.FadeIn = CreateDefaultSubobject<UFadeCurveFunctionExponential>(TEXT("FadeInFunction"));
 			}
-			else
+			else if (StartFadeCurve < 0.f)
 			{
-				FadeFunctions.FadeIn = CreateDefaultSubobject<UFadeCurveFunctionExponential>(TEXT("FadeInFunction"));
-				FadeFunctions.FadeIn->Duration = 0.f;
+				FadeFunctions.FadeIn = CreateDefaultSubobject<UFadeCurveFunctionSigmoid>(TEXT("FadeInFunction"));
+			}
+			else if (StartFadeCurve < WaveformTrasnformationTrimFadePrivate::Linear)
+			{
+				FadeFunctions.FadeIn = CreateDefaultSubobject<UFadeCurveFunctionLogarithmic>(TEXT("FadeInFunction"));
+			}
+			else if (StartFadeCurve == WaveformTrasnformationTrimFadePrivate::Linear)
+			{
+				FadeFunctions.FadeIn = CreateDefaultSubobject<UFadeFunctionLinear>(TEXT("FadeInFunction"));
 			}
-	}
 
+			FadeFunctions.FadeIn->Duration = StartFadeTime;
+			WaveformTrasnformationTrimFadePrivate::SetFadeCurve(FadeFunctions.FadeIn, StartFadeCurve, StartSCurveSharpness);
+
+			StartFadeTime = 0.f;
+		}
+		else
+		{
+			FadeFunctions.FadeIn = CreateDefaultSubobject<UFadeCurveFunctionExponential>(TEXT("FadeInFunction"));
+			FadeFunctions.FadeIn->Duration = 0.f;
+		}
+	}
 
 	if (FadeFunctions.FadeOut == nullptr)
 	{
-		PRAGMA_DISABLE_DEPRECATION_WARNINGS
-			if (EndFadeTime != 0.f)
+		if (EndFadeTime != 0.f)
+		{
+			if (EndFadeCurve > WaveformTrasnformationTrimFadePrivate::Linear)
 			{
-				if (EndFadeCurve > WaveformTrasnformationTrimFadePrivate::Linear)
-				{
-					FadeFunctions.FadeOut = CreateDefaultSubobject<UFadeCurveFunctionExponential>(TEXT("FadeOutFunction"));
-				}
-				else if (EndFadeCurve < 0.f)
-				{
-					FadeFunctions.FadeOut = CreateDefaultSubobject<UFadeCurveFunctionSigmoid>(TEXT("FadeOutFunction"));
-				}
-				else if (EndFadeCurve < WaveformTrasnformationTrimFadePrivate::Linear)
-				{
-					FadeFunctions.FadeOut = CreateDefaultSubobject<UFadeCurveFunctionLogarithmic>(TEXT("FadeOutFunction"));
-				}
-				else if (EndFadeCurve == WaveformTrasnformationTrimFadePrivate::Linear)
-				{
-					FadeFunctions.FadeOut = CreateDefaultSubobject<UFadeFunctionLinear>(TEXT("FadeOutFunction"));
-				}
-
-				FadeFunctions.FadeOut->Duration = EndFadeTime;
-				WaveformTrasnformationTrimFadePrivate::SetFadeCurve(FadeFunctions.FadeOut, EndFadeCurve, EndSCurveSharpness);
-
-				EndFadeTime = 0.f;
-				PRAGMA_ENABLE_DEPRECATION_WARNINGS
+				FadeFunctions.FadeOut = CreateDefaultSubobject<UFadeCurveFunctionExponential>(TEXT("FadeOutFunction"));
 			}
-			else
+			else if (EndFadeCurve < 0.f)
 			{
-				FadeFunctions.FadeOut = CreateDefaultSubobject<UFadeCurveFunctionExponential>(TEXT("FadeOutFunction"));
-				FadeFunctions.FadeOut->Duration = 0.f;
+				FadeFunctions.FadeOut = CreateDefaultSubobject<UFadeCurveFunctionSigmoid>(TEXT("FadeOutFunction"));
+			}
+			else if (EndFadeCurve < WaveformTrasnformationTrimFadePrivate::Linear)
+			{
+				FadeFunctions.FadeOut = CreateDefaultSubobject<UFadeCurveFunctionLogarithmic>(TEXT("FadeOutFunction"));
+			}
+			else if (EndFadeCurve == WaveformTrasnformationTrimFadePrivate::Linear)
+			{
+				FadeFunctions.FadeOut = CreateDefaultSubobject<UFadeFunctionLinear>(TEXT("FadeOutFunction"));
 			}
+
+			FadeFunctions.FadeOut->Duration = EndFadeTime;
+			WaveformTrasnformationTrimFadePrivate::SetFadeCurve(FadeFunctions.FadeOut, EndFadeCurve, EndSCurveSharpness);
+
+			EndFadeTime = 0.f;
+		}
+		else
+		{
+			FadeFunctions.FadeOut = CreateDefaultSubobject<UFadeCurveFunctionExponential>(TEXT("FadeOutFunction"));
+			FadeFunctions.FadeOut->Duration = 0.f;
+		}
+	}
+
+	if (FadeFunctions.FadeIn || FadeFunctions.FadeOut)
+	{
+		check(FadeTransformation);
+
+		FadeFunctions.MigrateToTransformationFadeFunction(FadeTransformation, GetMutableFadeRegion());
 	}
+	PRAGMA_ENABLE_DEPRECATION_WARNINGS
 }
 
 #if WITH_EDITOR
@@ -311,14 +355,19 @@ void UWaveformTransformationTrimFade::PostEditChangeProperty(FPropertyChangedEve
 	ensure(EndTime > StartTime);
 	const float MaxDuration = EndTime - StartTime;
 
-	if (FadeFunctions.FadeIn)
+	if (FadeTransformation)
 	{
-		FadeFunctions.FadeIn->Duration = FMath::Min(FadeFunctions.FadeIn->Duration, MaxDuration);
-	}
+		FadeTransformation->LimitFadeRegionsToOne();
 
-	if (FadeFunctions.FadeOut)
-	{
-		FadeFunctions.FadeOut->Duration = FMath::Min(FadeFunctions.FadeOut->Duration, MaxDuration);
+		if (GetFadeRegion().FadeIn)
+		{
+			GetFadeRegion().FadeIn->Duration = FMath::Min(GetFadeRegion().FadeIn->Duration, MaxDuration);
+		}
+
+		if (GetFadeRegion().FadeOut)
+		{
+			GetFadeRegion().FadeOut->Duration = FMath::Min(GetFadeRegion().FadeOut->Duration, MaxDuration);
+		}
 	}
 }
 
@@ -333,6 +382,7 @@ FString UWaveformTransformationTrimFade::GetTransformationHash() const
 	float FadeOutDuration = 0.f;
 	float FadeOutCurve = WaveformTrasnformationTrimFadePrivate::Linear;
 
+	PRAGMA_DISABLE_DEPRECATION_WARNINGS
 	if (FadeFunctions.FadeIn)
 	{
 		FadeInDuration = FadeFunctions.FadeIn->Duration;
@@ -358,6 +408,12 @@ FString UWaveformTransformationTrimFade::GetTransformationHash() const
 		FPCU::AppendHash(Hash, TEXT("TFOC"), FadeOutCurve);
 		FPCU::AppendHash(Hash, TEXT("TFOD"), FadeOutDuration);
 	}
+	PRAGMA_ENABLE_DEPRECATION_WARNINGS
+
+	if (FadeTransformation)
+	{
+		FPCU::AppendHash(Hash, TEXT("TFFT"), FadeTransformation->GetTransformationHash());
+	}
 
 	FPCU::AppendHash(Hash, TEXT("TFST"), StartTime);
 	FPCU::AppendHash(Hash, TEXT("TFET"), EndTime);
@@ -396,14 +452,7 @@ void UWaveformTransformationTrimFade::PostLoad()
 
 		StartFadeTime = 0.f;
 	}
-	PRAGMA_ENABLE_DEPRECATION_WARNINGS
-	else if (FadeFunctions.FadeIn == nullptr)
-	{
-		FadeFunctions.FadeIn = NewObject<UFadeCurveFunctionExponential>(this, UFadeCurveFunctionExponential::StaticClass(), NAME_None, RF_Transactional);
-		FadeFunctions.FadeIn->Duration = 0.f;
-	}
 
-	PRAGMA_DISABLE_DEPRECATION_WARNINGS
 	// Overwrite exiting fade out function data with deprecated curve properties to migrate the data to the new properties
 	if (EndFadeTime != 0.f)
 	{
@@ -429,25 +478,52 @@ void UWaveformTransformationTrimFade::PostLoad()
 
 		EndFadeTime = 0.f;
 	}
-	PRAGMA_ENABLE_DEPRECATION_WARNINGS
-	else if (FadeFunctions.FadeOut == nullptr)
+
+	if (FadeFunctions.FadeIn || FadeFunctions.FadeOut)
 	{
-		FadeFunctions.FadeOut = NewObject<UFadeCurveFunctionExponential>(this, UFadeCurveFunctionExponential::StaticClass(), NAME_None, RF_Transactional);
-		FadeFunctions.FadeOut->Duration = 0.f;
+		check(FadeTransformation);
+
+		FadeFunctions.MigrateToTransformationFadeFunction(FadeTransformation, GetMutableFadeRegion());
 	}
+	PRAGMA_ENABLE_DEPRECATION_WARNINGS
+
 }
 
 Audio::FTransformationPtr UWaveformTransformationTrimFade::CreateTransformation() const
 {
-	return MakeUnique<FWaveTransformationTrimFade>(FadeFunctions, StartTime, EndTime);
+	return MakeUnique<FWaveTransformationTrimFade>(FadeTransformation, StartTime, EndTime);
 }
 
 void UWaveformTransformationTrimFade::UpdateConfiguration(FWaveTransformUObjectConfiguration& InOutConfiguration)
 {
+	check(InOutConfiguration.SampleRate > 0.f);
 	UpdateDurationProperties(InOutConfiguration.EndTime - InOutConfiguration.StartTime);
 
 	InOutConfiguration.StartTime = StartTime;
 	InOutConfiguration.EndTime = EndTime;
+
+	// FadeTransformation needs StartTime = 0 and EndTime to be the original end time for proper offsets
+	FWaveTransformUObjectConfiguration InConfiguration = InOutConfiguration;
+	InConfiguration.StartTime = 0.f;
+	InConfiguration.EndTime = InOutConfiguration.GetOriginalNumFrames() / InOutConfiguration.SampleRate;
+
+	FadeTransformation->UpdateConfiguration(InConfiguration);
+}
+
+const FTransformationFadeFunctionData& UWaveformTransformationTrimFade::GetFadeRegion() const
+{
+	check(FadeTransformation);
+	check(FadeTransformation->GetFadeRegionsNum() == 1);
+
+	return FadeTransformation->GetFadeRegions()[0];
+}
+
+FTransformationFadeFunctionData& UWaveformTransformationTrimFade::GetMutableFadeRegion() const
+{
+	check(FadeTransformation);
+	check(FadeTransformation->GetFadeRegionsNum() == 1);
+
+	return FadeTransformation->GetMutableFadeRegions()[0];
 }
 
 PRAGMA_DISABLE_DEPRECATION_WARNINGS
@@ -623,3 +699,23 @@ const double UFadeCurveFunctionSigmoid::GetFadeOutCurveValue(const double FadeFr
 
 	return CurveValue;
 }
+
+void FFadeFunctionData::MigrateToTransformationFadeFunction(const TObjectPtr<UWaveformTransformationFade> InFadeTransformation, FTransformationFadeFunctionData& NewFadeRegion)
+{
+	check(InFadeTransformation);
+	check(InFadeTransformation->GetFadeRegionsNum() == 1);
+	
+	if(FadeIn)
+	{
+		WaveformTrasnformationTrimFadePrivate::MigrateToTransformationFadeFunction(InFadeTransformation, FadeIn, NewFadeRegion.FadeIn);
+
+		FadeIn = nullptr;
+	}
+
+	if (FadeOut)
+	{
+		WaveformTrasnformationTrimFadePrivate::MigrateToTransformationFadeFunction(InFadeTransformation, FadeOut, NewFadeRegion.FadeOut);
+
+		FadeOut = nullptr;
+	}
+}
diff --git a/Engine/Plugins/Editor/WaveformEditor/Source/WaveformTransformations/Private/WaveformTransformationTrimFadeCustomization.cpp b/Engine/Plugins/Editor/WaveformEditor/Source/WaveformTransformations/Private/WaveformTransformationTrimFadeCustomization.cpp
new file mode 100644
--- /dev/null
+++ b/Engine/Plugins/Editor/WaveformEditor/Source/WaveformTransformations/Private/WaveformTransformationTrimFadeCustomization.cpp
@@ -0,0 +1,165 @@
+// Copyright Epic Games, Inc. All Rights Reserved.
+
+#include "WaveformTransformationTrimFadeCustomization.h"
+
+#include "DetailWidgetRow.h"
+#include "IDetailChildrenBuilder.h"
+#include "IDetailGroup.h"
+#include "IPropertyUtilities.h"
+#include "PropertyHandle.h"
+#include "WaveformTransformationTrimFade.h"
+
+TSharedRef<IPropertyTypeCustomization> FWaveformTransformationTrimFadeCustomization::MakeInstance()
+{
+	return MakeShareable(new FWaveformTransformationTrimFadeCustomization);
+}
+
+void FWaveformTransformationTrimFadeCustomization::CustomizeHeader(TSharedRef<IPropertyHandle> PropertyHandle,
+	FDetailWidgetRow& HeaderRow, IPropertyTypeCustomizationUtils& CustomizationUtils)
+{
+	constexpr bool bDisplayDefaultPropertyButtons = false;
+
+	HeaderRow.NameContent()[PropertyHandle->CreatePropertyNameWidget()]
+	.ValueContent()[PropertyHandle->CreatePropertyValueWidget(bDisplayDefaultPropertyButtons)];
+}
+
+void FWaveformTransformationTrimFadeCustomization::CustomizeChildren(TSharedRef<IPropertyHandle> PropertyHandle,
+	IDetailChildrenBuilder& ChildBuilder, IPropertyTypeCustomizationUtils& CustomizationUtils)
+{
+	IDetailGroup& TrimGroup = ChildBuilder.AddGroup(TEXT("TrimGroup"), FText::FromString(TEXT("Trim")), true);
+
+	TSharedPtr<IPropertyHandle> StartTimeHandle = PropertyHandle->GetChildHandle(GET_MEMBER_NAME_CHECKED(UWaveformTransformationTrimFade, StartTime));
+	if (StartTimeHandle && StartTimeHandle->IsValidHandle())
+	{
+		TrimGroup.AddPropertyRow(StartTimeHandle.ToSharedRef());
+	}
+
+	TSharedPtr<IPropertyHandle> EndTimeHandle = PropertyHandle->GetChildHandle(GET_MEMBER_NAME_CHECKED(UWaveformTransformationTrimFade, EndTime));
+	if (EndTimeHandle && EndTimeHandle->IsValidHandle())
+	{
+		TrimGroup.AddPropertyRow(EndTimeHandle.ToSharedRef());
+	}
+
+	IDetailGroup& FadeGroup = ChildBuilder.AddGroup(TEXT("FadeGroup"), FText::FromString(TEXT("Fade")), true);
+
+	TSharedPtr<IPropertyHandle> FadeTransformationHandle = PropertyHandle->GetChildHandle(GET_MEMBER_NAME_CHECKED(UWaveformTransformationTrimFade, FadeTransformation));
+	if (FadeTransformationHandle && FadeTransformationHandle->IsValidHandle())
+	{
+		TSharedPtr<IPropertyHandle> FadeRegionArrayHandle = FadeTransformationHandle->GetChildHandle(UWaveformTransformationFade::GetFadeRegionsPropertyName());
+		if (FadeRegionArrayHandle && FadeRegionArrayHandle->IsValidHandle())
+		{
+			TSharedPtr<IPropertyHandleArray> FadeRegionArray = FadeRegionArrayHandle->AsArray();
+			if (FadeRegionArray)
+			{
+				uint32 NumRegions = 0;
+				FadeRegionArray->GetNumElements(NumRegions);
+
+				for (uint32 Index = 0; Index < NumRegions; Index++)
+				{
+					TSharedPtr<IPropertyHandle> FadeRegionHandle = FadeRegionArray->GetElement(Index);
+					if (FadeRegionHandle && FadeRegionHandle->IsValidHandle())
+					{
+						// Only show the Fade In and Out properties necessary for TrimFades
+						TSharedPtr<IPropertyHandle> FadeInHandle = FadeRegionHandle->GetChildHandle(GET_MEMBER_NAME_CHECKED(FTransformationFadeFunctionData, FadeIn));
+						if (FadeInHandle && FadeInHandle->IsValidHandle())
+						{
+							IDetailGroup& FadeInGroup = FadeGroup.AddGroup(TEXT("FadeInGroup"), FText::FromString(TEXT("Fade In")), true);
+							FadeInGroup.AddPropertyRow(FadeInHandle.ToSharedRef());
+
+							TSharedPtr<IPropertyHandle> FadeInDurHandle = FadeInHandle->GetChildHandle(GET_MEMBER_NAME_CHECKED(UTransformationFadeFunction, Duration));
+							if (FadeInDurHandle && FadeInDurHandle->IsValidHandle())
+							{
+								FadeInGroup.AddPropertyRow(FadeInDurHandle.ToSharedRef());
+							}
+
+							TSharedPtr<IPropertyHandle> FadeInOffsetHandle = FadeInHandle->GetChildHandle(GET_MEMBER_NAME_CHECKED(UTransformationFadeFunction, FrameOffset));
+							if (FadeInOffsetHandle && FadeInOffsetHandle->IsValidHandle())
+							{
+								FadeInOffsetHandle->MarkHiddenByCustomization();
+							}
+
+							UObject* FadeInObject = nullptr;
+							if (FadeInHandle->GetValue(FadeInObject) == FPropertyAccess::Success && FadeInObject)
+							{
+								if (FadeInObject->IsA<UTransformationFadeCurveFunctionExponential>())
+								{
+									TSharedPtr<IPropertyHandle> FadeCurveHandle = FadeInHandle->GetChildHandle(GET_MEMBER_NAME_CHECKED(UTransformationFadeCurveFunctionExponential, FadeCurve));
+									if (FadeCurveHandle && FadeCurveHandle->IsValidHandle())
+									{
+										FadeInGroup.AddPropertyRow(FadeCurveHandle.ToSharedRef());
+									}
+								}
+								else if (FadeInObject->IsA<UTransformationFadeCurveFunctionLogarithmic>())
+								{
+									TSharedPtr<IPropertyHandle> FadeCurveHandle = FadeInHandle->GetChildHandle(GET_MEMBER_NAME_CHECKED(UTransformationFadeCurveFunctionLogarithmic, FadeCurve));
+									if (FadeCurveHandle && FadeCurveHandle->IsValidHandle())
+									{
+										FadeInGroup.AddPropertyRow(FadeCurveHandle.ToSharedRef());
+									}
+								}
+								else if (FadeInObject->IsA<UTransformationFadeCurveFunctionSigmoid>())
+								{
+									TSharedPtr<IPropertyHandle> FadeCurveHandle = FadeInHandle->GetChildHandle(GET_MEMBER_NAME_CHECKED(UTransformationFadeCurveFunctionSigmoid, SigmoidFadeCurve));
+									if (FadeCurveHandle && FadeCurveHandle->IsValidHandle())
+									{
+										FadeInGroup.AddPropertyRow(FadeCurveHandle.ToSharedRef());
+									}
+								}
+							}
+						}
+
+						TSharedPtr<IPropertyHandle> FadeOutHandle = FadeRegionHandle->GetChildHandle(GET_MEMBER_NAME_CHECKED(FTransformationFadeFunctionData, FadeOut));
+						if (FadeOutHandle && FadeOutHandle->IsValidHandle())
+						{
+							IDetailGroup& FadeOutGroup = FadeGroup.AddGroup(TEXT("FadeOutGroup"), FText::FromString(TEXT("Fade Out")), true);
+							FadeOutGroup.AddPropertyRow(FadeOutHandle.ToSharedRef());
+
+							TSharedPtr<IPropertyHandle> FadeInDurHandle = FadeOutHandle->GetChildHandle(GET_MEMBER_NAME_CHECKED(UTransformationFadeFunction, Duration));
+							if (FadeInDurHandle && FadeInDurHandle->IsValidHandle())
+							{
+								FadeOutGroup.AddPropertyRow(FadeInDurHandle.ToSharedRef());
+							}
+
+							TSharedPtr<IPropertyHandle> FadeOutOffsetHandle = FadeOutHandle->GetChildHandle(GET_MEMBER_NAME_CHECKED(UTransformationFadeFunction, FrameOffset));
+							if (FadeOutOffsetHandle && FadeOutOffsetHandle->IsValidHandle())
+							{
+								FadeOutOffsetHandle->MarkHiddenByCustomization();
+							}
+
+							UObject* FadeOutObject = nullptr;
+							if (FadeOutHandle->GetValue(FadeOutObject) == FPropertyAccess::Success && FadeOutObject)
+							{
+								if (FadeOutObject->IsA<UTransformationFadeCurveFunctionExponential>())
+								{
+									TSharedPtr<IPropertyHandle> FadeCurveHandle = FadeOutHandle->GetChildHandle(GET_MEMBER_NAME_CHECKED(UTransformationFadeCurveFunctionExponential, FadeCurve));
+									if (FadeCurveHandle && FadeCurveHandle->IsValidHandle())
+									{
+										FadeOutGroup.AddPropertyRow(FadeCurveHandle.ToSharedRef());
+									}
+								}
+								else if (FadeOutObject->IsA<UTransformationFadeCurveFunctionLogarithmic>())
+								{
+									TSharedPtr<IPropertyHandle> FadeCurveHandle = FadeOutHandle->GetChildHandle(GET_MEMBER_NAME_CHECKED(UTransformationFadeCurveFunctionLogarithmic, FadeCurve));
+									if (FadeCurveHandle && FadeCurveHandle->IsValidHandle())
+									{
+										FadeOutGroup.AddPropertyRow(FadeCurveHandle.ToSharedRef());
+									}
+								}
+								else if (FadeOutObject->IsA<UTransformationFadeCurveFunctionSigmoid>())
+								{
+									TSharedPtr<IPropertyHandle> FadeCurveHandle = FadeOutHandle->GetChildHandle(GET_MEMBER_NAME_CHECKED(UTransformationFadeCurveFunctionSigmoid, SigmoidFadeCurve));
+									if (FadeCurveHandle && FadeCurveHandle->IsValidHandle())
+									{
+										FadeOutGroup.AddPropertyRow(FadeCurveHandle.ToSharedRef());
+									}
+								}
+							}
+						}
+					}
+				}
+			}
+		}
+
+	}
+
+}
\ No newline at end of file
diff --git a/Engine/Plugins/Editor/WaveformEditor/Source/WaveformTransformations/Private/WaveformTransformationTrimFadeCustomization.h b/Engine/Plugins/Editor/WaveformEditor/Source/WaveformTransformations/Private/WaveformTransformationTrimFadeCustomization.h
new file mode 100644
--- /dev/null
+++ b/Engine/Plugins/Editor/WaveformEditor/Source/WaveformTransformations/Private/WaveformTransformationTrimFadeCustomization.h
@@ -0,0 +1,18 @@
+// Copyright Epic Games, Inc. All Rights Reserved.
+
+#pragma once
+#include "IPropertyTypeCustomization.h"
+
+/**
+ * Using IPropertyTypeCustomization to edit the details view of the TrimFade transformation and remove unneeded information in the waveform editor
+ */
+class FWaveformTransformationTrimFadeCustomization : public IPropertyTypeCustomization
+{
+public:
+	static TSharedRef<IPropertyTypeCustomization> MakeInstance();
+
+	virtual void CustomizeHeader(TSharedRef<IPropertyHandle> PropertyHandle, FDetailWidgetRow& HeaderRow,
+								 IPropertyTypeCustomizationUtils& CustomizationUtils) override;
+	virtual void CustomizeChildren(TSharedRef<IPropertyHandle> PropertyHandle, IDetailChildrenBuilder& ChildBuilder,
+								   IPropertyTypeCustomizationUtils& CustomizationUtils) override;
+};
diff --git a/Engine/Plugins/Editor/WaveformEditor/Source/WaveformTransformations/Private/WaveformTransformationsModule.cpp b/Engine/Plugins/Editor/WaveformEditor/Source/WaveformTransformations/Private/WaveformTransformationsModule.cpp
--- a/Engine/Plugins/Editor/WaveformEditor/Source/WaveformTransformations/Private/WaveformTransformationsModule.cpp
+++ b/Engine/Plugins/Editor/WaveformEditor/Source/WaveformTransformations/Private/WaveformTransformationsModule.cpp
@@ -5,7 +5,9 @@
 #include "Modules/ModuleManager.h"
 #include "PropertyEditorModule.h"
 #include "WaveformTransformationLog.h"
+#include "WaveformTransformationFadeCustomization.h"
 #include "WaveformTransformationMarkersObjectCustomization.h"
+#include "WaveformTransformationTrimFadeCustomization.h"
 
 DEFINE_LOG_CATEGORY(LogWaveformTransformation);
 
@@ -25,6 +27,16 @@ void FWaveformTransformationsModule::StartupModule()
 		FOnGetPropertyTypeCustomizationInstance::CreateStatic(&FWaveformTransformationMarkersObjectCustomization::MakeInstance)
 	);
 
+	PropertyEditorModule->RegisterCustomPropertyTypeLayout(
+		"WaveformTransformationTrimFade",
+		FOnGetPropertyTypeCustomizationInstance::CreateStatic(&FWaveformTransformationTrimFadeCustomization::MakeInstance)
+	);
+
+	PropertyEditorModule->RegisterCustomPropertyTypeLayout(
+		"WaveformTransformationFade",
+		FOnGetPropertyTypeCustomizationInstance::CreateStatic(&FWaveformTransformationFadeCustomization::MakeInstance)
+	);
+
 	PropertyEditorModule->NotifyCustomizationModuleChanged();
 }
 
@@ -37,6 +49,8 @@ void FWaveformTransformationsModule::ShutdownModule()
 		FPropertyEditorModule& PropertyEditorModule = FModuleManager::GetModuleChecked<FPropertyEditorModule>("PropertyEditor");
 
 		PropertyEditorModule.UnregisterCustomPropertyTypeLayout("WaveformTransformationMarkers");
+		PropertyEditorModule.UnregisterCustomPropertyTypeLayout("WaveformTransformationTrimFade");
+		PropertyEditorModule.UnregisterCustomPropertyTypeLayout("WaveformTransformationFade");
 		PropertyEditorModule.NotifyCustomizationModuleChanged();
 	}
 }
diff --git a/Engine/Plugins/Editor/WaveformEditor/Source/WaveformTransformations/Public/WaveformTransformationFade.h b/Engine/Plugins/Editor/WaveformEditor/Source/WaveformTransformations/Public/WaveformTransformationFade.h
new file mode 100644
--- /dev/null
+++ b/Engine/Plugins/Editor/WaveformEditor/Source/WaveformTransformations/Public/WaveformTransformationFade.h
@@ -0,0 +1,202 @@
+// Copyright Epic Games, Inc. All Rights Reserved.
+
+#pragma once
+
+#include "UObject/Object.h"
+#include "IWaveformTransformation.h"
+
+#include "WaveformTransformationFade.generated.h"
+
+#define UE_API WAVEFORMTRANSFORMATIONS_API
+
+UCLASS(MinimalAPI, Abstract, EditInlineNew, DefaultToInstanced, CollapseCategories, NotBlueprintable)
+class UTransformationFadeFunction : public UObject
+{
+	GENERATED_BODY()
+public:
+
+	virtual const double GetFadeInCurveValue(const double FadeFraction) const { return FadeFraction; }
+	virtual const double GetFadeOutCurveValue(const double FadeFraction) const { return 1.0 - FadeFraction; }
+
+	UPROPERTY(EditAnywhere, Category = "Fade", BlueprintReadWrite, meta = (ClampMin = 0.0, DisplayName = "Fade Duration"))
+	float Duration = 0.5f;
+
+	UPROPERTY(EditAnywhere, Category = "Fade", BlueprintReadWrite, meta = (ClampMin = 0.0, DisplayName = "Fade Frame Offset"))
+	int64 FrameOffset = 0;
+};
+
+UCLASS(MinimalAPI, Abstract, EditInlineNew, DefaultToInstanced, CollapseCategories, NotBlueprintable)
+class UTransformationFadeCurveFunction : public UTransformationFadeFunction
+{
+	GENERATED_BODY()
+public:
+
+	virtual const double GetFadeInCurveValue(const double FadeFraction) const override { return FadeFraction; }
+	virtual const double GetFadeOutCurveValue(const double FadeFraction) const override { return 1.0 - FadeFraction; }
+
+	virtual void SetFadeCurve(const float InFadeCurve) {}
+	virtual const float GetFadeCurve() const { return 0.f; }
+	virtual const FText GetFadeCurvePropertyName() const { return FText::FromName(NAME_None); }
+
+protected:
+	virtual const float GetFadeCurveMin() const { return 0.f; }
+	virtual const float GetFadeCurveMax() const { return 1.f; }
+};
+
+UCLASS(MinimalAPI, NotBlueprintable, meta = (DisplayName = "Linear"))
+class UTransformationFadeFunctionLinear : public UTransformationFadeFunction
+{
+	GENERATED_BODY()
+public:
+	
+	UE_API virtual const double GetFadeInCurveValue(const double FadeFraction) const override;
+	UE_API virtual const double GetFadeOutCurveValue(const double FadeFraction) const override;
+};
+
+UCLASS(MinimalAPI, NotBlueprintable, meta = (DisplayName = "Exponential"))
+class UTransformationFadeCurveFunctionExponential : public UTransformationFadeCurveFunction
+{
+	GENERATED_BODY()
+public:
+
+	UE_API virtual const double GetFadeInCurveValue(const double FadeFraction) const override;
+	UE_API virtual const double GetFadeOutCurveValue(const double FadeFraction) const override;
+
+	virtual void SetFadeCurve(const float InFadeCurve) override { FadeCurve = FMath::Clamp(InFadeCurve, GetFadeCurveMin(), GetFadeCurveMax()); }
+	virtual const float GetFadeCurve() const override { return FadeCurve; }
+	virtual const FText GetFadeCurvePropertyName() const override { return FText::FromName(GET_MEMBER_NAME_CHECKED(UTransformationFadeCurveFunctionExponential, FadeCurve)); }
+
+	UPROPERTY(EditAnywhere, Category = "Fade", meta = (ClampMin = 1, ClampMax = 10.0, DisplayName = "Fade Curve"))
+	float FadeCurve = 3.f;
+
+protected:
+	virtual const float GetFadeCurveMin() const override { return 1.f; }
+	virtual const float GetFadeCurveMax() const override { return 10.f; }
+};
+
+UCLASS(MinimalAPI, NotBlueprintable, meta = (DisplayName = "Logarithmic"))
+class UTransformationFadeCurveFunctionLogarithmic : public UTransformationFadeCurveFunction
+{
+	GENERATED_BODY()
+public:
+
+	UE_API virtual const double GetFadeInCurveValue(const double FadeFraction) const override;
+	UE_API virtual const double GetFadeOutCurveValue(const double FadeFraction) const override;
+
+	virtual void SetFadeCurve(const float InFadeCurve) override { FadeCurve = FMath::Clamp(InFadeCurve, GetFadeCurveMin(), GetFadeCurveMax()); }
+	virtual const float GetFadeCurve() const override { return FadeCurve; }
+	virtual const FText GetFadeCurvePropertyName() const override { return FText::FromName(GET_MEMBER_NAME_CHECKED(UTransformationFadeCurveFunctionLogarithmic, FadeCurve)); }
+
+	UPROPERTY(EditAnywhere, Category = "Fade", meta = (ClampMin = 0.0, ClampMax = 1, DisplayName = "Fade Curve"))
+	float FadeCurve = 0.25f;
+};
+
+UCLASS(MinimalAPI, NotBlueprintable, meta = (DisplayName = "Sigmoid"))
+class UTransformationFadeCurveFunctionSigmoid : public UTransformationFadeCurveFunction
+{
+	GENERATED_BODY()
+public:
+
+	UE_API virtual const double GetFadeInCurveValue(const double FadeFraction) const override;
+	UE_API virtual const double GetFadeOutCurveValue(const double FadeFraction) const override;
+
+	virtual void SetFadeCurve(const float InFadeCurve) override { SigmoidFadeCurve = FMath::Clamp(InFadeCurve, GetFadeCurveMin(), GetFadeCurveMax()); }
+	virtual const float GetFadeCurve() const override { return SigmoidFadeCurve; }
+	virtual const FText GetFadeCurvePropertyName() const override { return FText::FromName(GET_MEMBER_NAME_CHECKED(UTransformationFadeCurveFunctionSigmoid, SigmoidFadeCurve)); }
+
+	UPROPERTY(EditAnywhere, Category = "Fade", meta = (ClampMin = 0.0, ClampMax = 1, DisplayName = "S-Fade Curve"))
+	float SigmoidFadeCurve = 0.6f;
+};
+
+USTRUCT()
+struct FTransformationFadeFunctionData
+{
+	GENERATED_BODY()
+
+	UPROPERTY(EditAnywhere, Category = "Fade", Instanced, meta = (DisplayName = "Fade In Type"))
+	TObjectPtr<UTransformationFadeFunction> FadeIn;
+
+	UPROPERTY(EditAnywhere, Category = "Fade", Instanced, meta = (DisplayName = "Fade Out Type"))
+	TObjectPtr<UTransformationFadeFunction> FadeOut;
+
+	UPROPERTY(EditAnywhere, Category = "Fade", meta = (DisplayName = "Link Fade Durations"))
+	bool bLinkDurations = false;
+};
+
+UENUM()
+enum class EWaveEditorTransformationFadeMode : uint8
+{
+	Linear = 0,
+	Exponential,
+	Logarithmic,
+	Sigmoid
+};
+
+class FWaveTransformationFade : public Audio::IWaveTransformation
+{
+public:
+	UE_API explicit FWaveTransformationFade(const TArray<FTransformationFadeFunctionData>& InFadeRegions, const int64 InStartFrameOffset, const int64 InEndFrameOffset);
+	UE_API virtual void ProcessAudio(Audio::FWaveformTransformationWaveInfo& InOutWaveInfo) const override;
+
+private:
+
+	TArray<FTransformationFadeFunctionData> FadeRegions;
+	int64 StartFrameOffset = 0;
+	int64 EndFrameOffset = 0;
+};
+
+// Maintains an array of fade regions that can be placed in any location on the waveform.
+// Each region has a start and end fade with optional offsets from the start and end of the waveform
+// This transformation can be used in other transformation (ex. TrimFade)
+UCLASS(MinimalAPI)
+class UWaveformTransformationFade final : public UWaveformTransformationBase
+{
+	GENERATED_BODY()
+public:
+
+	UWaveformTransformationFade(const FObjectInitializer& ObjectInitializer);
+
+#if WITH_EDITOR
+	UE_API virtual void PostEditChangeProperty(FPropertyChangedEvent& PropertyChangedEvent) override;
+
+	UE_API virtual FString GetTransformationHash() const override;
+#endif // WITH_EDITOR;
+
+	static UE_API const TMap<EWaveEditorTransformationFadeMode, TSubclassOf<UTransformationFadeFunction>> FadeModeToFadeFunctionMap;
+
+	UE_API virtual Audio::FTransformationPtr CreateTransformation() const override;
+
+	UE_API void UpdateConfiguration(FWaveTransformUObjectConfiguration& InOutConfiguration) override;
+
+	const float GetSampleRate() const { return SampleRate; }
+
+	const int32 GetFadeRegionsNum() const { return FadeRegions.Num(); }
+	const TArray<FTransformationFadeFunctionData>& GetFadeRegions() const { return FadeRegions; }
+	TArray<FTransformationFadeFunctionData>& GetMutableFadeRegions() { return FadeRegions; }
+
+	// Remove all fade regions other than the first one.
+	// Logs a warning if it removes any fade regions
+	UE_API void LimitFadeRegionsToOne();
+
+	static FName GetFadeRegionsPropertyName()
+	{
+		return GET_MEMBER_NAME_CHECKED(UWaveformTransformationFade, FadeRegions);
+	}
+
+protected:
+	UPROPERTY(EditAnywhere, Category = "Fade", meta = (DisplayName = "Fade Regions"))
+	TArray<FTransformationFadeFunctionData> FadeRegions;
+
+	
+private:	
+
+	UE_API void UpdateDurationProperties(const float InAvailableDuration);
+	float AvailableWaveformDuration = -1.f;
+	float SampleRate = 0.f;
+
+	int64 StartFrameOffset = 0;
+	int64 EndFrameOffset = 0;
+	int64 OriginalNumFrames = 0;
+};
+
+#undef UE_API
diff --git a/Engine/Plugins/Editor/WaveformEditor/Source/WaveformTransformations/Public/WaveformTransformationTrimFade.h b/Engine/Plugins/Editor/WaveformEditor/Source/WaveformTransformations/Public/WaveformTransformationTrimFade.h
--- a/Engine/Plugins/Editor/WaveformEditor/Source/WaveformTransformations/Public/WaveformTransformationTrimFade.h
+++ b/Engine/Plugins/Editor/WaveformEditor/Source/WaveformTransformations/Public/WaveformTransformationTrimFade.h
@@ -4,13 +4,14 @@
 
 #include "UObject/Object.h"
 #include "IWaveformTransformation.h"
+#include "WaveformTransformationFade.h"
 
 #include "WaveformTransformationTrimFade.generated.h"
 
 #define UE_API WAVEFORMTRANSFORMATIONS_API
 
 UCLASS(MinimalAPI, Abstract, EditInlineNew, DefaultToInstanced, CollapseCategories, NotBlueprintable)
-class UFadeFunction : public UObject
+class UE_DEPRECATED(5.8, "UFadeFunction has been deprecated. Use UTransformationFadeFunction.") UFadeFunction : public UObject
 {
 	GENERATED_BODY()
 public:
@@ -22,8 +23,9 @@ public:
 	float Duration = 0.5f;
 };
 
+PRAGMA_DISABLE_DEPRECATION_WARNINGS
 UCLASS(MinimalAPI, Abstract, EditInlineNew, DefaultToInstanced, CollapseCategories, NotBlueprintable)
-class UFadeCurveFunction : public UFadeFunction
+class UE_DEPRECATED(5.8, "UFadeFunction has been deprecated. Use UTransformationFadeFunction.") UFadeCurveFunction : public UFadeFunction
 {
 	GENERATED_BODY()
 public:
@@ -41,7 +43,7 @@ protected:
 };
 
 UCLASS(MinimalAPI, NotBlueprintable, meta = (DisplayName = "Linear"))
-class UFadeFunctionLinear : public UFadeFunction
+class UE_DEPRECATED(5.8, "UFadeFunction has been deprecated. Use UTransformationFadeFunction.") UFadeFunctionLinear : public UFadeFunction
 {
 	GENERATED_BODY()
 public:
@@ -51,7 +53,7 @@ public:
 };
 
 UCLASS(MinimalAPI, NotBlueprintable, meta = (DisplayName = "Exponential"))
-class UFadeCurveFunctionExponential : public UFadeCurveFunction
+class UE_DEPRECATED(5.8, "UFadeFunction has been deprecated. Use UTransformationFadeFunction.") UFadeCurveFunctionExponential : public UFadeCurveFunction
 {
 	GENERATED_BODY()
 public:
@@ -72,7 +74,7 @@ protected:
 };
 
 UCLASS(MinimalAPI, NotBlueprintable, meta = (DisplayName = "Logarithmic"))
-class UFadeCurveFunctionLogarithmic : public UFadeCurveFunction
+class UE_DEPRECATED(5.8, "UFadeFunction has been deprecated. Use UTransformationFadeFunction.") UFadeCurveFunctionLogarithmic : public UFadeCurveFunction
 {
 	GENERATED_BODY()
 public:
@@ -89,7 +91,7 @@ public:
 };
 
 UCLASS(MinimalAPI, NotBlueprintable, meta = (DisplayName = "Sigmoid"))
-class UFadeCurveFunctionSigmoid : public UFadeCurveFunction
+class UE_DEPRECATED(5.8, "UFadeFunction has been deprecated. Use UTransformationFadeFunction.") UFadeCurveFunctionSigmoid : public UFadeCurveFunction
 {
 	GENERATED_BODY()
 public:
@@ -106,7 +108,7 @@ public:
 };
 
 USTRUCT()
-struct FFadeFunctionData
+struct UE_DEPRECATED(5.8, "UFadeFunction has been deprecated. Use UTransformationFadeFunction.") FFadeFunctionData
 {
 	GENERATED_BODY()
 
@@ -115,28 +117,31 @@ struct FFadeFunctionData
 
 	UPROPERTY(EditAnywhere, Category = "Fade", Instanced)
 	TObjectPtr<UFadeFunction> FadeOut;
+
+	void MigrateToTransformationFadeFunction(const TObjectPtr<UWaveformTransformationFade> InFadeTransformation, FTransformationFadeFunctionData& NewFadeRegion);
 };
 
 UENUM()
-enum class EWaveEditorFadeMode : uint8
+enum class UE_DEPRECATED(5.8, "UFadeFunction has been deprecated. Use UTransformationFadeFunction.") EWaveEditorFadeMode : uint8
 {
 	Linear = 0,
 	Exponential,
 	Logarithmic,
 	Sigmoid
 };
+PRAGMA_ENABLE_DEPRECATION_WARNINGS
 
 class FWaveTransformationTrimFade : public Audio::IWaveTransformation
 {
 public:
-	UE_API explicit FWaveTransformationTrimFade(const FFadeFunctionData& InFadeFunctions, double InStartTime, double InEndTime);
+	UE_API explicit FWaveTransformationTrimFade(TObjectPtr<UWaveformTransformationFade> InFadeTransformation, double InStartTime, double InEndTime);
 	UE_API virtual void ProcessAudio(Audio::FWaveformTransformationWaveInfo& InOutWaveInfo) const override;
 
 	virtual constexpr Audio::ETransformationPriority FileChangeLengthPriority() const override { return Audio::ETransformationPriority::High; }
 
 private:
 
-	FFadeFunctionData FadeFunctions;
+	TObjectPtr<UWaveformTransformationFade> FadeTransformation;
 
 	double StartTime = 0.0;
 	double EndTime = 0.0;
@@ -164,73 +169,84 @@ public:
 	UPROPERTY(EditAnywhere, Category = "Trim")
 	double EndTime = -1.0;
 
-	UPROPERTY(EditAnywhere, Category = "Fade", meta = (DisplayName = "Fade Functions"))
+	UPROPERTY(EditAnywhere, Instanced, Category = "Fade", meta = (DisplayName = "Fade Functions"))
+	TObjectPtr<UWaveformTransformationFade> FadeTransformation = nullptr;
+
+	PRAGMA_DISABLE_DEPRECATION_WARNINGS
+	UE_DEPRECATED(5.8, "UFadeFunction has been deprecated. Use UTransformationFadeFunction.")
+	UPROPERTY(EditAnywhere, Category = "Fade", meta = (DisplayName = "Fade Functions", EditCondition = "false", EditConditionHides))
 	FFadeFunctionData FadeFunctions;
+	PRAGMA_ENABLE_DEPRECATION_WARNINGS
 
-	UE_DEPRECATED(5.7, "Individual fade properties have been deprecated. Use FadeFunctions instead.")
+	UE_DEPRECATED(5.7, "Individual fade properties have been deprecated. Use UTransformationFadeFunction instead.")
 	/*
 	 * Relative to StartTime
-	 * Individual fade properties have been deprecated. Use FadeFunctions instead.
+	 * Individual fade properties have been deprecated. Use UTransformationFadeFunction instead.
 	 */ 
 	UPROPERTY(EditAnywhere, Category = "Fade", meta = (ClampMin = 0.0, DisplayName = "Fade-In Duration", EditCondition = "StartFadeTime != 0", EditConditionHides))
 	float StartFadeTime = 0.f;
 	
-	UE_DEPRECATED(5.7, "Individual fade properties have been deprecated. Use FadeFunctions instead.")
+	UE_DEPRECATED(5.7, "Individual fade properties have been deprecated. Use UTransformationFadeFunction instead.")
 	/*
-	 * Individual fade properties have been deprecated. Use FadeFunctions instead.
+	 * Individual fade properties have been deprecated. Use UTransformationFadeFunction instead.
 	 */
 	UPROPERTY(EditAnywhere, Category = "Fade", meta=(ClampMin = -0.1, ClampMax = 10.0, DisplayName = "Fade-In Curve", EditCondition = "StartFadeTime != 0", EditConditionHides))
 	float StartFadeCurve = 1.f;
 
-	UE_DEPRECATED(5.7, "Individual fade properties have been deprecated. Use FadeFunctions instead.")
+	UE_DEPRECATED(5.7, "Individual fade properties have been deprecated. Use UTransformationFadeFunction instead.")
 	/*
-	 * Individual fade properties have been deprecated. Use FadeFunctions instead.
+	 * Individual fade properties have been deprecated. Use UTransformationFadeFunction instead.
 	 */
 	UPROPERTY(EditAnywhere, Category = "Fade", meta = (ClampMin = 0, ClampMax = 1.0, DisplayName = "S-Curve Sharpness", EditCondition = "StartFadeTime != 0 && StartFadeCurve < 0", EditConditionHides))
 	float StartSCurveSharpness = 0.75f;
 
-	UE_DEPRECATED(5.7, "Individual fade properties have been deprecated. Use FadeFunctions instead.")
+	UE_DEPRECATED(5.7, "Individual fade properties have been deprecated. Use UTransformationFadeFunction instead.")
 	/*
 	 * Relative to EndTime. Fade out starts at EndTime - EndFadeTime and ends at EndTime
-	 * Individual fade properties have been deprecated. Use FadeFunctions instead.
+	 * Individual fade properties have been deprecated. Use UTransformationFadeFunction instead.
 	 */
 	UPROPERTY(EditAnywhere, Category = "Fade", meta = (ClampMin = 0.0, DisplayName = "Fade-Out Duration", EditCondition = "EndFadeTime != 0", EditConditionHides))
 	float EndFadeTime = 0.f;
 
-	UE_DEPRECATED(5.7, "Individual fade properties have been deprecated. Use FadeFunctions instead.")
+	UE_DEPRECATED(5.7, "Individual fade properties have been deprecated. Use UTransformationFadeFunction instead.")
 	/*
-	 * Individual fade properties have been deprecated. Use FadeFunctions instead.
+	 * Individual fade properties have been deprecated. Use UTransformationFadeFunction instead.
 	 */
 	UPROPERTY(EditAnywhere, Category = "Fade", meta=(ClampMin = -0.1, ClampMax = 10.0, DisplayName = "Fade-Out Curve", EditCondition = "EndFadeTime != 0", EditConditionHides))
 	float EndFadeCurve = 1.f;
 
-	UE_DEPRECATED(5.7, "Individual fade properties have been deprecated. Use FadeFunctions instead.")
+	UE_DEPRECATED(5.7, "Individual fade properties have been deprecated. Use UTransformationFadeFunction instead.")
 	/*
-	 * Individual fade properties have been deprecated. Use FadeFunctions instead.
+	 * Individual fade properties have been deprecated. Use UTransformationFadeFunction instead.
 	 */
 	UPROPERTY(EditAnywhere, Category = "Fade", meta = (ClampMin = 0, ClampMax = 1.0, DisplayName = "S-Curve Sharpness", EditCondition = "EndFadeTime != 0 && EndFadeCurve < 0", EditConditionHides))
 	float EndSCurveSharpness = 0.75;
 
-	UE_DEPRECATED(5.7, "Individual fade properties have been deprecated. Use FadeFunctions instead.")
+	UE_DEPRECATED(5.7, "Individual fade properties have been deprecated. Use UTransformationFadeFunction instead.")
 	static UE_API const float MinFadeCurve;
 
-	UE_DEPRECATED(5.7, "Individual fade properties have been deprecated. Use FadeFunctions instead.")
+	UE_DEPRECATED(5.7, "Individual fade properties have been deprecated. Use UTransformationFadeFunction instead.")
 	static UE_API const float MaxFadeCurve;
 	
-	UE_DEPRECATED(5.7, "Individual fade properties have been deprecated. Use FadeFunctions instead.")
+	PRAGMA_DISABLE_DEPRECATION_WARNINGS
+	UE_DEPRECATED(5.7, "Individual fade properties have been deprecated. Use UTransformationFadeFunction instead.")
 	static UE_API const TMap<EWaveEditorFadeMode, float> FadeModeToCurveValueMap;
 
 	static UE_API const TMap<EWaveEditorFadeMode, TSubclassOf<UFadeFunction>> FadeModeToFadeFunctionMap;
+	PRAGMA_ENABLE_DEPRECATION_WARNINGS
 
 	UE_API virtual Audio::FTransformationPtr CreateTransformation() const override;
 
 	UE_API void UpdateConfiguration(FWaveTransformUObjectConfiguration& InOutConfiguration) override;
 	virtual constexpr Audio::ETransformationPriority GetTransformationPriority() const { return Audio::ETransformationPriority::High; }	
 
-	UE_DEPRECATED(5.7, "Individual fade properties have been deprecated. Use FadeFunctions::FadeIn::GetFadeInCurveValue instead.")
+	UE_API const FTransformationFadeFunctionData& GetFadeRegion() const;
+	UE_API FTransformationFadeFunctionData& GetMutableFadeRegion() const;
+
+	UE_DEPRECATED(5.7, "Individual fade properties have been deprecated. Use UTransformationFadeFunction instead.")
 	static UE_API const double GetFadeInCurveValue(const float StartFadeCurve, const double FadeFraction, const float SCurveSharpness = 0);
 
-	UE_DEPRECATED(5.7, "Individual fade properties have been deprecated. Use FadeFunctions::FadeOut::GetFadeOutCurveValue instead.")
+	UE_DEPRECATED(5.7, "Individual fade properties have been deprecated. Use UTransformationFadeFunction instead.")
 	static UE_API const double GetFadeOutCurveValue(const float EndFadeCurve, const double FadeFraction, const float SCurveSharpness = 0);
 	
 private:	
diff --git a/Engine/Plugins/Editor/WaveformEditor/Source/WaveformTransformationsWidgets/Private/WaveformTransformationFadeRenderer.cpp b/Engine/Plugins/Editor/WaveformEditor/Source/WaveformTransformationsWidgets/Private/WaveformTransformationFadeRenderer.cpp
new file mode 100644
--- /dev/null
+++ b/Engine/Plugins/Editor/WaveformEditor/Source/WaveformTransformationsWidgets/Private/WaveformTransformationFadeRenderer.cpp
@@ -0,0 +1,978 @@
+// Copyright Epic Games, Inc. All Rights Reserved.
+
+#include "WaveformTransformationFadeRenderer.h"
+
+#include "AudioWidgetsStyle.h"
+#include "Editor.h"
+#include "Framework/Application/SlateApplication.h"
+#include "Layout/Geometry.h"
+#include "PropertyHandle.h"
+#include "Rendering/DrawElements.h"
+#include "Sound/SoundWave.h"
+#include "WaveformTransformationLog.h"
+#include "Widgets/Images/SImage.h"
+#include "Widgets/Input/SButton.h"
+#include "Widgets/SWidget.h"
+#include "Widgets/Text/STextBlock.h"
+#include "Blueprint/WidgetLayoutLibrary.h"
+
+FWaveformTransformationFadeRenderer::~FWaveformTransformationFadeRenderer()
+{
+	if (FSlateApplication::IsInitialized())
+	{
+		FSlateApplication::Get().OnFocusChanging().Remove(PopupHandle);
+		FSlateApplication::Get().OnApplicationActivationStateChanged().Remove(ApplicationActivationStateHandle);
+	}
+
+	OnFadeModeChanged.Unbind();
+}
+
+int32 FWaveformTransformationFadeRenderer::OnPaint(const FPaintArgs& Args, const FGeometry& AllottedGeometry, const FSlateRect& MyCullingRect, FSlateWindowElementList& OutDrawElements, int32 LayerId, const FWidgetStyle& InWidgetStyle, bool bParentEnabled) const
+{
+	LayerId = DrawOffsetHandles(AllottedGeometry, OutDrawElements, LayerId);
+	LayerId = DrawFadeCurves(AllottedGeometry, OutDrawElements, LayerId);
+
+	return LayerId;
+}
+
+int32 FWaveformTransformationFadeRenderer::DrawOffsetHandles(const FGeometry& AllottedGeometry, FSlateWindowElementList& OutDrawElements, int32 LayerId) const
+{
+	const float HandleWidth = InteractionHandleSize;
+
+	for (const FFadeRegion& FadeRegion : FadeRegions)
+	{
+		FPaintGeometry HandleGeometry = AllottedGeometry.ToPaintGeometry(FVector2f(HandleWidth, HandleWidth), FSlateLayoutTransform(FVector2f(FadeRegion.FadeInStartX, AllottedGeometry.Size.Y - HandleWidth)));
+
+		FSlateDrawElement::MakeBox(
+			OutDrawElements,
+			++LayerId,
+			HandleGeometry,
+			&RoundedBoxBrush,
+			ESlateDrawEffect::None,
+			FLinearColor::Green
+		);
+
+		HandleGeometry = AllottedGeometry.ToPaintGeometry(FVector2f(HandleWidth, HandleWidth), FSlateLayoutTransform(FVector2f(FadeRegion.FadeOutEndX - HandleWidth, AllottedGeometry.Size.Y - HandleWidth)));
+
+		FSlateDrawElement::MakeBox(
+			OutDrawElements,
+			++LayerId,
+			HandleGeometry,
+			&RoundedBoxBrush,
+			ESlateDrawEffect::None,
+			FLinearColor::Red
+		);
+	}
+
+	return LayerId;
+}
+
+int32 FWaveformTransformationFadeRenderer::DrawFadeCurves(const FGeometry& AllottedGeometry, FSlateWindowElementList& OutDrawElements, int32 LayerId) const
+{
+	const FLinearColor FadeColor = FLinearColor::Gray;
+	const float HandleWidth = InteractionHandleSize;
+	const float HandleHeight = InteractionHandleSize;
+	const float HandleOffset = 5.0f;
+
+	const FLinearColor HandleColor = FLinearColor::White;
+	const FLinearColor HandleSelectedColor = FLinearColor::Yellow;
+
+	for (const FFadeRegion& FadeRegion : FadeRegions)
+	{
+		const bool IsFadeInSelected = FadeRegion.FadeInEndInteractionXRange.Contains(MousePosition.X) && FadeRegion.FadeInEndInteractionYRange.Contains(MousePosition.Y);
+		const bool IsFadeOutSelected = FadeRegion.FadeOutStartInteractionXRange.Contains(MousePosition.X) && FadeRegion.FadeOutStartInteractionYRange.Contains(MousePosition.Y);
+
+		if (FadeRegion.FadeInCurvePoints.Num() > 0)
+		{
+			FSlateDrawElement::MakeLines(
+				OutDrawElements,
+				++LayerId,
+				AllottedGeometry.ToPaintGeometry(),
+				FadeRegion.FadeInCurvePoints,
+				ESlateDrawEffect::None,
+				FadeColor
+			);
+
+			const float	HandleStart = FadeRegion.FadeInCurvePoints.Last(0).X;
+			const FPaintGeometry HandleGeometry = AllottedGeometry.ToPaintGeometry(FVector2f(HandleWidth, HandleHeight), FSlateLayoutTransform(FVector2f(HandleStart - HandleOffset, 0)));
+			const FSlateBrush* HandleBrush = FAppStyle::GetBrush("WorldBrowser.DirectionXPositive");
+
+			FSlateDrawElement::MakeBox(
+				OutDrawElements,
+				++LayerId,
+				HandleGeometry,
+				HandleBrush,
+				ESlateDrawEffect::DisabledEffect,
+				IsFadeInSelected ? HandleSelectedColor : HandleColor
+			);
+		}
+		else
+		{
+			FPaintGeometry HandleGeometry = AllottedGeometry.ToPaintGeometry(FVector2f(HandleWidth, HandleHeight), FSlateLayoutTransform(FVector2f(FadeRegion.FadeInStartX- HandleOffset, 0)));
+			const FSlateBrush* HandleBrush = FAppStyle::GetBrush("WorldBrowser.DirectionXPositive");
+
+			FSlateDrawElement::MakeBox(
+				OutDrawElements,
+				++LayerId,
+				HandleGeometry,
+				HandleBrush,
+				ESlateDrawEffect::DisabledEffect,
+				IsFadeInSelected ? HandleSelectedColor : HandleColor
+			);
+		}
+
+		if (FadeRegion.FadeOutCurvePoints.Num() > 0)
+		{
+			FSlateDrawElement::MakeLines(
+				OutDrawElements,
+				++LayerId,
+				AllottedGeometry.ToPaintGeometry(),
+				FadeRegion.FadeOutCurvePoints,
+				ESlateDrawEffect::None,
+				FadeColor
+			);
+
+			const float	HandleStart = FadeRegion.FadeOutCurvePoints[0].X - HandleWidth;
+			const FPaintGeometry HandleGeometry = AllottedGeometry.ToPaintGeometry(FVector2f(HandleWidth, HandleHeight), FSlateLayoutTransform(FVector2f(HandleStart + HandleOffset, 0)));
+			const FSlateBrush* HandleBrush = FAppStyle::GetBrush("WorldBrowser.DirectionXNegative");
+
+			FSlateDrawElement::MakeBox(
+				OutDrawElements,
+				++LayerId,
+				HandleGeometry,
+				HandleBrush,
+				ESlateDrawEffect::DisabledEffect,
+				IsFadeOutSelected ? HandleSelectedColor : HandleColor
+			);
+		}
+		else
+		{
+			FPaintGeometry HandleGeometry = AllottedGeometry.ToPaintGeometry(FVector2f(HandleWidth, HandleHeight), FSlateLayoutTransform(FVector2f(FadeRegion.FadeOutEndX - HandleWidth + HandleOffset, 0)));
+			const FSlateBrush* HandleBrush = FAppStyle::GetBrush("WorldBrowser.DirectionXNegative");
+
+			FSlateDrawElement::MakeBox(
+				OutDrawElements,
+				++LayerId,
+				HandleGeometry,
+				HandleBrush,
+				ESlateDrawEffect::DisabledEffect,
+				IsFadeOutSelected ? HandleSelectedColor : HandleColor
+			);
+		}
+	}
+
+	return LayerId;
+}
+
+void FWaveformTransformationFadeRenderer::GenerateFadeCurves(const FGeometry& AllottedGeometry)
+{
+	check(StrongFade);
+	check(FadeRegions.Num() == StrongFade->GetFadeRegionsNum());
+	check(TransformationWaveInfo.NumChannels > 0);
+
+	const int64 NumFrames = (TransformationWaveInfo.NumSamplesAvailable / TransformationWaveInfo.NumChannels);
+
+	for (int32 Index = 0; Index < FadeRegions.Num(); Index++)
+	{
+		check(FadeRegions[Index].FadeInStartX >= 0);
+		check(FadeRegions[Index].FadeInEndX >= 0);
+		check(StrongFade->GetFadeRegions()[Index].FadeIn);
+		check(StrongFade->GetFadeRegions()[Index].FadeOut);
+
+		if (StrongFade->GetFadeRegions()[Index].FadeIn->Duration > 0.f && PixelsPerFrame > 0.0)
+		{
+			const float FadeInTime = StrongFade->GetFadeRegions()[Index].FadeIn->Duration;
+			const float FadeInFrames = FadeInTime * TransformationWaveInfo.SampleRate;
+			const float FadeInPixelLength = FadeInFrames * PixelsPerFrame;
+			check(FadeInPixelLength > 0.f);
+
+			const uint32 DisplayedFadeInPixelLength = 
+				(FadeRegions[Index].FadeInStartX < FadeRegions[Index].FadeInEndX) ? (FadeRegions[Index].FadeInEndX - FadeRegions[Index].FadeInStartX) : 0;
+
+			if (DisplayedFadeInPixelLength > (NumFrames * PixelsPerFrame))
+			{
+				// Mid property change, the DisplayedFadeInPixelLength can become invalid here. It will be updated next tick
+				return;
+			}
+
+			FadeRegions[Index].FadeInCurvePoints.SetNumUninitialized(DisplayedFadeInPixelLength);
+
+			if (DisplayedFadeInPixelLength > 0)
+			{
+				for (uint32 Pixel = 0; Pixel < DisplayedFadeInPixelLength; ++Pixel)
+				{
+					const double FadeFraction = FMath::Clamp(static_cast<float>(Pixel) / FadeInPixelLength, 0.f, 1.f);
+					const double CurveFunction = StrongFade->GetFadeRegions()[Index].FadeIn->GetFadeInCurveValue(FadeFraction);
+
+					const bool bNeedsFlickerPrevention = (Pixel == DisplayedFadeInPixelLength - 1) && DisplayedFadeInPixelLength != 1;
+
+					// Prevent fade line from flickering if the math doesn't work out to the last Y value being 0 as the position is changing
+					// This is due to pixel rounding errors. The slight distortion caused by this correction isn't noticeable
+					const double CurveValue = bNeedsFlickerPrevention ? 0.f : 1.f - CurveFunction;
+
+					const uint32 XCoordinate = Pixel + FadeRegions[Index].FadeInStartX;
+					FadeRegions[Index].FadeInCurvePoints[Pixel] = FVector2D(XCoordinate, CurveValue * AllottedGeometry.Size.Y);
+				}
+
+				check(FadeRegions[Index].FadeInCurvePoints.Num() > 0);
+				check(FadeRegions[Index].FadeInCurvePoints[0].Y == AllottedGeometry.Size.Y);
+			}
+		}
+		else
+		{
+			FadeRegions[Index].FadeInCurvePoints.SetNumUninitialized(0);
+		}
+
+		if (StrongFade->GetFadeRegions()[Index].FadeOut->Duration > 0.f && PixelsPerFrame > 0.0)
+		{
+			const float FadeOutTime = StrongFade->GetFadeRegions()[Index].FadeOut->Duration;
+			const float FadeOutFrames = FadeOutTime * TransformationWaveInfo.SampleRate;
+			const float FadeOutPixelLength = FadeOutFrames * PixelsPerFrame;
+			check(FadeOutPixelLength > 0.f);
+
+			const uint32 DisplayedFadeOutPixelLength = 
+				(FadeRegions[Index].FadeOutStartX < FadeRegions[Index].FadeOutEndX) ? (FadeRegions[Index].FadeOutEndX - FadeRegions[Index].FadeOutStartX) : 0;
+
+			if (DisplayedFadeOutPixelLength > (NumFrames * PixelsPerFrame))
+			{
+				// Mid property change, the DisplayedFadeInPixelLength can become invalid here. It will be updated next tick
+				return;
+			}
+
+			FadeRegions[Index].FadeOutCurvePoints.SetNumUninitialized(DisplayedFadeOutPixelLength);
+
+			if (DisplayedFadeOutPixelLength > 0)
+			{
+				for (uint32 Pixel = 0; Pixel < DisplayedFadeOutPixelLength; ++Pixel)
+				{
+					const double FadeFraction = FMath::Clamp(static_cast<float>(Pixel) / FadeOutPixelLength, 0.f, 1.f);
+					const double CurveFunction = StrongFade->GetFadeRegions()[Index].FadeOut->GetFadeOutCurveValue(FadeFraction);
+
+					const bool bNeedsFlickerPrevention = (Pixel == DisplayedFadeOutPixelLength - 1) && DisplayedFadeOutPixelLength != 1;
+
+					// Prevent fade line from flickering if the math doesn't work out to the last Y value being equal to AllottedGeometry.Size.Y as the 
+					// position is changing. This is due to pixel rounding errors. The slight distortion caused by this correction isn't noticeable
+					const double CurveValue = bNeedsFlickerPrevention ? 1.f : 1.f - CurveFunction;
+
+					const uint32 XCoordinate = Pixel + FadeRegions[Index].FadeOutStartX;
+					FadeRegions[Index].FadeOutCurvePoints[Pixel] = FVector2D(XCoordinate, CurveValue * AllottedGeometry.Size.Y);
+				}
+
+				check(FadeRegions[Index].FadeOutCurvePoints.Num() > 0);
+				check(FadeRegions[Index].FadeOutCurvePoints[0].Y == 0.f);
+			}
+		}
+		else
+		{
+			FadeRegions[Index].FadeOutCurvePoints.SetNumUninitialized(0);
+		}
+	}
+}
+
+FCursorReply FWaveformTransformationFadeRenderer::OnCursorQuery(const FGeometry& MyGeometry, const FPointerEvent& CursorEvent) const
+{
+	const FVector2D LocalCursorPosition = GetLocalCursorPosition(CursorEvent, MyGeometry);
+
+	switch (FadeInteractionTypeAndIndex.Key)
+	{
+	case FWaveformTransformationFadeRenderer::EFadeInteractionType::ScrubbingFadeInStart:
+	case FWaveformTransformationFadeRenderer::EFadeInteractionType::ScrubbingFadeInEnd:
+	case FWaveformTransformationFadeRenderer::EFadeInteractionType::ScrubbingFadeOutStart:
+	case FWaveformTransformationFadeRenderer::EFadeInteractionType::ScrubbingFadeOutEnd:
+		return FCursorReply::Cursor(EMouseCursor::GrabHandClosed);
+	default:
+		break;
+	}
+
+	if (IsCursorInFadeInteractionRange(LocalCursorPosition, MyGeometry))
+	{
+		return FCursorReply::Cursor(EMouseCursor::GrabHand);
+	}
+
+	return FCursorReply::Unhandled();
+}
+
+void FWaveformTransformationFadeRenderer::SetWaveformTransformation(TObjectPtr<UWaveformTransformationBase> InTransformation)
+{
+	FWaveformTransformationRendererBase::SetWaveformTransformation(InTransformation);
+	
+	StrongFade = TStrongObjectPtr<UWaveformTransformationFade>(CastChecked<UWaveformTransformationFade>(InTransformation.Get()));
+
+	FadeRegions.Reset();
+	FadeRegions.SetNum(StrongFade->GetFadeRegionsNum());
+}
+
+FReply FWaveformTransformationFadeRenderer::OnMouseWheel(SWidget& OwnerWidget, const FGeometry& MyGeometry, const FPointerEvent& MouseEvent)
+{
+	check(StrongFade);
+	check(FadeRegions.Num() == StrongFade->GetFadeRegionsNum());
+
+	// Prevent stomping on other mouse transactions
+	if (FadeInteractionTypeAndIndex.Key != EFadeInteractionType::None)
+	{
+		return FReply::Handled();
+	}
+
+	const FVector2D LocalCursorPosition = GetLocalCursorPosition(MouseEvent, MyGeometry);
+
+	for (int32 Index = 0; Index < FadeRegions.Num(); Index++)
+	{
+		check(StrongFade->GetFadeRegions()[Index].FadeIn);
+		check(StrongFade->GetFadeRegions()[Index].FadeOut);
+
+		if (FadeRegions[Index].FadeInEndInteractionXRange.Contains(LocalCursorPosition.X) 
+			&& FadeRegions[Index].FadeInEndInteractionYRange.Contains(LocalCursorPosition.Y) 
+			&& StrongFade->GetFadeRegions()[Index].FadeIn.IsA<UTransformationFadeCurveFunction>())
+		{
+			const float FadeCurveDelta = MouseEvent.GetWheelDelta() * MouseWheelStep;
+			FText OutPropertyName = FText::FromName(NAME_None);
+
+			TObjectPtr<UTransformationFadeCurveFunction> FadeIn =
+				CastChecked<UTransformationFadeCurveFunction>(StrongFade->GetFadeRegions()[Index].FadeIn);
+			OutPropertyName = FadeIn->GetFadeCurvePropertyName();
+
+			if (!GIsTransacting)
+			{
+				TransactionResult =
+					BeginTransaction(TEXT("PropertyEditor"), FText::Format(NSLOCTEXT("PropertyEditor", "SetNumericPropertyTransaction", "Edit {0}"), OutPropertyName), StrongFade.Get());
+
+				if (TransactionResult == INDEX_NONE)
+				{
+					UE_LOG(LogWaveformTransformation, Warning, TEXT("Begin FadeTransformation OnMouseWheel Transaction Failed"));
+				}
+			}
+
+			StrongFade->Modify();
+			StrongFade->GetFadeRegions()[Index].FadeIn->Modify();
+
+			const float FadeCurveValue = FadeIn->GetFadeCurve() + FadeCurveDelta;
+			FadeIn->SetFadeCurve(FadeCurveValue);
+
+			if (TransactionResult != INDEX_NONE)
+			{
+				EndTransaction();
+				TransactionResult = INDEX_NONE;
+			}
+
+			StrongFade->OnTransformationChanged.ExecuteIfBound(true);
+
+			return FReply::Handled();
+		}
+
+		if (FadeRegions[Index].FadeOutStartInteractionXRange.Contains(LocalCursorPosition.X) 
+			&& FadeRegions[Index].FadeOutStartInteractionYRange.Contains(LocalCursorPosition.Y) 
+			&& StrongFade->GetFadeRegions()[Index].FadeOut.IsA<UTransformationFadeCurveFunction>())
+		{
+			const float FadeCurveDelta = MouseEvent.GetWheelDelta() * MouseWheelStep;
+			FText OutPropertyName = FText::FromName(NAME_None);
+
+			TObjectPtr<UTransformationFadeCurveFunction> FadeOut =
+				CastChecked<UTransformationFadeCurveFunction>(StrongFade->GetFadeRegions()[Index].FadeOut);
+			OutPropertyName = FadeOut->GetFadeCurvePropertyName();
+			
+			if (!GIsTransacting)
+			{
+				TransactionResult =
+					BeginTransaction(TEXT("PropertyEditor"), FText::Format(NSLOCTEXT("PropertyEditor", "SetNumericPropertyTransaction", "Edit {0}"), OutPropertyName), StrongFade.Get());
+
+				if (TransactionResult == INDEX_NONE)
+				{
+					UE_LOG(LogWaveformTransformation, Warning, TEXT("Begin FadeTransformation OnMouseWheel Transaction Failed"));
+				}
+			}
+
+			StrongFade->Modify();
+			StrongFade->GetFadeRegions()[Index].FadeOut->Modify();
+
+			const float FadeCurveValue = FadeOut->GetFadeCurve() + FadeCurveDelta;
+			FadeOut->SetFadeCurve(FadeCurveValue);
+
+			if (TransactionResult != INDEX_NONE)
+			{
+				EndTransaction();
+				TransactionResult = INDEX_NONE;
+			}
+
+			StrongFade->OnTransformationChanged.ExecuteIfBound(true);
+
+			return FReply::Handled();
+		}
+	}
+
+	return FReply::Unhandled();
+}
+
+FVector2D FWaveformTransformationFadeRenderer::GetLocalCursorPosition(const FPointerEvent& MouseEvent, const FGeometry& EventGeometry) const
+{
+	const FVector2D ScreenSpacePosition = MouseEvent.GetScreenSpacePosition();
+	return  EventGeometry.AbsoluteToLocal(ScreenSpacePosition);
+}
+
+double FWaveformTransformationFadeRenderer::ConvertXRatioToTime(const float InRatio) const
+{
+	check(TransformationWaveInfo.NumChannels > 0);
+	check(TransformationWaveInfo.SampleRate > 0.f);
+
+	const float NumFrames = TransformationWaveInfo.TotalNumSamples / TransformationWaveInfo.NumChannels;
+	const float FrameSelected = NumFrames * InRatio;
+	return FrameSelected / TransformationWaveInfo.SampleRate;
+}
+
+void FWaveformTransformationFadeRenderer::Tick(const FGeometry& AllottedGeometry, const double InCurrentTime, const float InDeltaTime)
+{ 
+	check(TransformationWaveInfo.NumChannels > 0);
+	check(StrongFade);
+
+	const float NumFrames = TransformationWaveInfo.TotalNumSamples / TransformationWaveInfo.NumChannels;
+
+	check(NumFrames > 0);
+	PixelsPerFrame = FMath::Max(AllottedGeometry.GetLocalSize().X / NumFrames, 0.0);
+
+	const int32 NumFadeRegions = StrongFade->GetFadeRegionsNum();
+	if (NumFadeRegions != FadeRegions.Num())
+	{
+		FadeRegions.SetNum(NumFadeRegions);
+	}
+
+	check(NumFadeRegions == FadeRegions.Num());
+
+	for (int32 Index = 0; Index < NumFadeRegions; Index++)
+	{
+		check(StrongFade->GetFadeRegions()[Index].FadeIn);
+		check(StrongFade->GetFadeRegions()[Index].FadeOut);
+
+		FadeRegions[Index].FadeInStartX = static_cast<uint32>(FMath::Clamp(StrongFade->GetFadeRegions()[Index].FadeIn->FrameOffset * PixelsPerFrame, 0, UINT32_MAX));
+		FadeRegions[Index].FadeOutEndX = static_cast<uint32>(FMath::Clamp((NumFrames * PixelsPerFrame) - (StrongFade->GetFadeRegions()[Index].FadeOut->FrameOffset * PixelsPerFrame), 0, UINT32_MAX));
+		FadeRegions[Index].FadeInEndX = static_cast<uint32>(FMath::Clamp(StrongFade->GetFadeRegions()[Index].FadeIn->Duration * TransformationWaveInfo.SampleRate * PixelsPerFrame + FadeRegions[Index].FadeInStartX, 0, UINT32_MAX));
+		FadeRegions[Index].FadeOutStartX = static_cast<uint32>(FMath::Clamp(FadeRegions[Index].FadeOutEndX - (StrongFade->GetFadeRegions()[Index].FadeOut->Duration * TransformationWaveInfo.SampleRate * PixelsPerFrame), 0, UINT32_MAX));
+	}
+
+	FVector2D MouseAbsolutePosition = UWidgetLayoutLibrary::GetMousePositionOnPlatform();
+	MousePosition = AllottedGeometry.AbsoluteToLocal(MouseAbsolutePosition);
+
+	GenerateFadeCurves(AllottedGeometry);
+	UpdateInteractionRange(AllottedGeometry);
+}
+
+void FWaveformTransformationFadeRenderer::UpdateInteractionRange(const FGeometry& AllottedGeometry)
+{
+	for (FFadeRegion& FadeRegion : FadeRegions)
+	{
+		FadeRegion.FadeInStartInteractionXRange.SetLowerBoundValue(FadeRegion.FadeInStartX - InteractionHandleSize);
+		FadeRegion.FadeInStartInteractionXRange.SetUpperBoundValue(FadeRegion.FadeInStartX + InteractionHandleSize);
+		FadeRegion.FadeInEndInteractionXRange.SetLowerBoundValue(FadeRegion.FadeInEndX - InteractionHandleSize);
+		FadeRegion.FadeInEndInteractionXRange.SetUpperBoundValue(FadeRegion.FadeInEndX + InteractionHandleSize);
+		FadeRegion.FadeInStartInteractionYRange.SetLowerBoundValue(AllottedGeometry.GetLocalSize().Y - InteractionHandleSize);
+		FadeRegion.FadeInStartInteractionYRange.SetUpperBoundValue(AllottedGeometry.GetLocalSize().Y);
+		FadeRegion.FadeInEndInteractionYRange.SetLowerBoundValue(0);
+		FadeRegion.FadeInEndInteractionYRange.SetUpperBoundValue(InteractionHandleSize);
+
+		FadeRegion.FadeOutStartInteractionXRange.SetLowerBoundValue(FadeRegion.FadeOutStartX - InteractionHandleSize);
+		FadeRegion.FadeOutStartInteractionXRange.SetUpperBoundValue(FadeRegion.FadeOutStartX + InteractionHandleSize);
+		FadeRegion.FadeOutEndInteractionXRange.SetLowerBoundValue(FadeRegion.FadeOutEndX - InteractionHandleSize);
+		FadeRegion.FadeOutEndInteractionXRange.SetUpperBoundValue(FadeRegion.FadeOutEndX + InteractionHandleSize);
+		FadeRegion.FadeOutStartInteractionYRange.SetLowerBoundValue(0);
+		FadeRegion.FadeOutStartInteractionYRange.SetUpperBoundValue(InteractionHandleSize);
+		FadeRegion.FadeOutEndInteractionYRange.SetLowerBoundValue(AllottedGeometry.GetLocalSize().Y - InteractionHandleSize);
+		FadeRegion.FadeOutEndInteractionYRange.SetUpperBoundValue(AllottedGeometry.GetLocalSize().Y);
+	}
+}
+
+bool FWaveformTransformationFadeRenderer::IsCursorInFadeInteractionRange(const FVector2D& InLocalCursorPosition, const FGeometry& WidgetGeometry) const
+{
+	for (const FFadeRegion& FadeRegion : FadeRegions)
+	{
+		if ((FadeRegion.FadeInStartInteractionXRange.Contains(InLocalCursorPosition.X) && FadeRegion.FadeInStartInteractionYRange.Contains(InLocalCursorPosition.Y))
+			|| (FadeRegion.FadeInEndInteractionXRange.Contains(InLocalCursorPosition.X) && FadeRegion.FadeInEndInteractionYRange.Contains(InLocalCursorPosition.Y))
+			|| (FadeRegion.FadeOutStartInteractionXRange.Contains(InLocalCursorPosition.X) && FadeRegion.FadeOutStartInteractionYRange.Contains(InLocalCursorPosition.Y))
+			|| (FadeRegion.FadeOutEndInteractionXRange.Contains(InLocalCursorPosition.X) && FadeRegion.FadeOutEndInteractionYRange.Contains(InLocalCursorPosition.Y)))
+		{
+			return true;
+		}
+	}
+
+	return false;
+}
+
+FReply FWaveformTransformationFadeRenderer::OnMouseButtonDown(SWidget& OwnerWidget, const FGeometry& MyGeometry, const FPointerEvent& MouseEvent)
+{
+	check(StrongFade);
+	const FVector2D LocalCursorPosition = GetLocalCursorPosition(MouseEvent, MyGeometry);
+
+	const FKey MouseButton = MouseEvent.GetEffectingButton();
+
+	FadeInteractionTypeAndIndex = GetInteractionTypeAndIndexFromCursorPosition(LocalCursorPosition, MouseButton, MyGeometry);
+
+	if (FadeInteractionTypeAndIndex.Key != EFadeInteractionType::None)
+	{
+		if (!GIsTransacting)
+		{
+			TransactionResult =
+				BeginTransaction(TEXT("PropertyEditor"), FText::Format(NSLOCTEXT("PropertyEditor", "SetNumericPropertyTransaction", "Edit {0}"), GetPropertyEditedByCurrentInteraction()), StrongFade.Get());
+
+			if (TransactionResult == INDEX_NONE)
+			{
+				UE_LOG(LogWaveformTransformation, Warning, TEXT("Begin FadeTransformation OnMouseButtonDown Transaction Failed"));
+			}
+		}
+
+		if (StrongFade->GetFadeRegions().IsValidIndex(FadeInteractionTypeAndIndex.Value))
+		{
+			check(StrongFade->GetFadeRegions()[FadeInteractionTypeAndIndex.Value].FadeIn);
+			check(StrongFade->GetFadeRegions()[FadeInteractionTypeAndIndex.Value].FadeOut);
+
+			if ((FadeInteractionTypeAndIndex.Key == EFadeInteractionType::ScrubbingFadeInEnd 
+				|| FadeInteractionTypeAndIndex.Key == EFadeInteractionType::ScrubbingFadeInStart))
+			{
+				StrongFade->GetFadeRegions()[FadeInteractionTypeAndIndex.Value].FadeIn->Modify();
+			}
+			else if ((FadeInteractionTypeAndIndex.Key == EFadeInteractionType::ScrubbingFadeOutEnd 
+				|| FadeInteractionTypeAndIndex.Key == EFadeInteractionType::ScrubbingFadeOutStart))
+			{
+				StrongFade->GetFadeRegions()[FadeInteractionTypeAndIndex.Value].FadeOut->Modify();
+			}
+		}
+
+		StrongFade->Modify();
+
+		return FReply::Handled().CaptureMouse(OwnerWidget.AsShared()).PreventThrottling();
+	}
+
+	return FReply::Unhandled();
+}
+
+
+FReply FWaveformTransformationFadeRenderer::OnMouseMove(SWidget& OwnerWidget, const FGeometry& MyGeometry, const FPointerEvent& MouseEvent)
+{
+	check(StrongFade);
+	if (MouseEvent.IsMouseButtonDown(EKeys::LeftMouseButton) && FadeInteractionTypeAndIndex.Key != EFadeInteractionType::None)
+	{
+		SetPropertyValueDependingOnInteractionType(MouseEvent, MyGeometry);
+
+		StrongFade->PostEditChange();
+
+		return FReply::Handled().CaptureMouse(OwnerWidget.AsShared());
+	}
+
+	return FReply::Unhandled();
+}
+
+FReply FWaveformTransformationFadeRenderer::OnMouseButtonUp(SWidget& OwnerWidget, const FGeometry& MyGeometry, const FPointerEvent& MouseEvent)
+{
+	check(StrongFade);
+	if (FadeInteractionTypeAndIndex.Key != EFadeInteractionType::None)
+	{
+		SetPropertyValueDependingOnInteractionType(MouseEvent, MyGeometry);
+
+		if (TransactionResult != INDEX_NONE)
+		{
+			EndTransaction();
+			TransactionResult = INDEX_NONE;
+		}
+
+		StrongFade->PostEditChange();
+
+		FadeInteractionTypeAndIndex.Key = EFadeInteractionType::None;
+		FadeInteractionTypeAndIndex.Value = 0;
+		
+		return FReply::Handled().ReleaseMouseCapture();
+	}
+
+	return FReply::Unhandled();
+}
+
+void FWaveformTransformationFadeRenderer::SetPropertyValueDependingOnInteractionType(const FPointerEvent& MouseEvent, const FGeometry& WidgetGeometry)
+{
+	check(StrongFade);
+	check(PixelsPerFrame > 0.0);
+	check(StrongFade->GetSampleRate() > 0);
+	check(WidgetGeometry.GetLocalSize().X > 0.f);
+	const FVector2D LocalCursorPosition = GetLocalCursorPosition(MouseEvent, WidgetGeometry);
+	const float LocalCursorXRatio = FMath::Clamp(LocalCursorPosition.X / WidgetGeometry.GetLocalSize().X, 0.f, 1.f);
+	const double SelectedTime = ConvertXRatioToTime(LocalCursorXRatio);
+	const float NumFrames = TransformationWaveInfo.TotalNumSamples / TransformationWaveInfo.NumChannels;
+	const float MaxDuration = NumFrames / StrongFade->GetSampleRate();
+	int64 FrameOffset = 0;
+
+	if (StrongFade->GetFadeRegions().IsValidIndex(FadeInteractionTypeAndIndex.Value))
+	{
+		check(StrongFade->GetFadeRegions()[FadeInteractionTypeAndIndex.Value].FadeIn);
+		check(StrongFade->GetFadeRegions()[FadeInteractionTypeAndIndex.Value].FadeOut);
+
+		switch (FadeInteractionTypeAndIndex.Key)
+		{
+		case FWaveformTransformationFadeRenderer::EFadeInteractionType::None:
+			break;
+		case FWaveformTransformationFadeRenderer::EFadeInteractionType::ScrubbingFadeInStart:
+		{
+			StrongFade->GetFadeRegions()[FadeInteractionTypeAndIndex.Value].FadeIn->FrameOffset = FMath::FloorToInt64(LocalCursorPosition.X / PixelsPerFrame);
+
+			break;
+		}
+		case FWaveformTransformationFadeRenderer::EFadeInteractionType::ScrubbingFadeInEnd:
+		{
+			float StartFadeDurationValue = 0;
+			StartFadeDurationValue = FMath::Clamp(SelectedTime - (FadeRegions[FadeInteractionTypeAndIndex.Value].FadeInStartX / PixelsPerFrame) / StrongFade->GetSampleRate(), 0.f, MaxDuration);
+			StrongFade->GetFadeRegions()[FadeInteractionTypeAndIndex.Value].FadeIn->Duration = StartFadeDurationValue;
+
+			break;
+		}
+		case FWaveformTransformationFadeRenderer::EFadeInteractionType::ScrubbingFadeOutStart:
+		{
+			float EndFadeDurationValue = 0;
+			const float diff = ((FadeRegions[FadeInteractionTypeAndIndex.Value].FadeOutEndX / PixelsPerFrame) / StrongFade->GetSampleRate()) - SelectedTime;
+			EndFadeDurationValue = FMath::Clamp(diff, 0.f, MaxDuration);
+			StrongFade->GetFadeRegions()[FadeInteractionTypeAndIndex.Value].FadeOut->Duration = EndFadeDurationValue;
+
+			break;
+		}
+		case FWaveformTransformationFadeRenderer::EFadeInteractionType::ScrubbingFadeOutEnd:
+		{
+			StrongFade->GetFadeRegions()[FadeInteractionTypeAndIndex.Value].FadeOut->FrameOffset = FMath::FloorToInt64(NumFrames - LocalCursorPosition.X / PixelsPerFrame);
+
+			break;
+		}
+		case FWaveformTransformationFadeRenderer::EFadeInteractionType::RightClickFadeIn:
+		case FWaveformTransformationFadeRenderer::EFadeInteractionType::RightClickFadeOut:
+			ShowSelectFadeModeMenuAtCursor(WidgetGeometry, MouseEvent);
+			break;
+		default:
+			break;
+		}
+	}
+}
+
+void FWaveformTransformationFadeRenderer::ShowSelectFadeModeMenuAtCursor(const FGeometry& MyGeometry, const FPointerEvent& MouseEvent)
+{
+	if (!FSlateApplication::IsInitialized())
+	{
+		UE_LOG(LogWaveformTransformation, Warning, TEXT("FSlateApplication not initialized in WaveformTransformationTrimFadeRenderer."));
+
+		return;
+	}
+
+	const FVector2D LocalWindowMaxPosition = MyGeometry.GetAbsolutePosition() + MyGeometry.GetAbsoluteSize();
+	const FVector2D LocalCursorPosition = MouseEvent.GetScreenSpacePosition();
+	
+	if (TSharedPtr<SWindow> LockedFadeModeMenuWindow = FadeModeMenuWindow.Pin())
+	{
+		FSlateApplication::Get().RequestDestroyWindow(LockedFadeModeMenuWindow.ToSharedRef());
+		FSlateApplication::Get().OnFocusChanging().Remove(PopupHandle); 
+		FSlateApplication::Get().OnApplicationActivationStateChanged().Remove(ApplicationActivationStateHandle);
+	}
+
+	TSharedRef<SVerticalBox> MenuContent = SNew(SVerticalBox);
+
+	TSharedRef<SWindow> MenuWindow = SNew(SWindow)
+		.AutoCenter(EAutoCenter::None)
+		.ScreenPosition(LocalCursorPosition)
+		.SizingRule(ESizingRule::Autosized)
+		.SupportsTransparency(EWindowTransparency::None)
+		.IsPopupWindow(true)
+		.CreateTitleBar(false)
+		[
+			MenuContent
+		];
+
+	FadeModeMenuWindow = MenuWindow;
+	TWeakObjectPtr<UWaveformTransformationFade> InWeakFade = StrongFade.Get();
+
+	TSharedPtr<FDelegateHandle> ApplicationActivationStateHandlePtr = MakeShared<FDelegateHandle>();
+	*ApplicationActivationStateHandlePtr = FSlateApplication::Get().OnApplicationActivationStateChanged().AddLambda([InFadeModeMenuWindow = FadeModeMenuWindow, ApplicationActivationStateHandlePtr](bool isActive)
+		{
+			if (!isActive)
+			{
+				if (TSharedPtr<SWindow> LockedFadeModeMenuWindow = InFadeModeMenuWindow.Pin())
+				{
+					FSlateApplication::Get().RequestDestroyWindow(LockedFadeModeMenuWindow.ToSharedRef());
+					FSlateApplication::Get().OnApplicationActivationStateChanged().Remove(*ApplicationActivationStateHandlePtr);
+				}
+			}
+		});
+
+	if (ApplicationActivationStateHandle.IsValid())
+	{
+		FSlateApplication::Get().OnApplicationActivationStateChanged().Remove(ApplicationActivationStateHandle);
+	}
+
+	ApplicationActivationStateHandle = *ApplicationActivationStateHandlePtr;
+
+	// If focus is lost on the popup, destroy it to prevent popups hanging around
+	TSharedPtr<FDelegateHandle> PopupHandlePtr = MakeShared<FDelegateHandle>();
+	*PopupHandlePtr = FSlateApplication::Get().OnFocusChanging().AddLambda([InFadeModeMenuWindow = FadeModeMenuWindow, PopupHandlePtr,
+		InApplicationActivationStateHandle = ApplicationActivationStateHandle](const FFocusEvent& FocusEvent, const FWeakWidgetPath& WeakWidgetPath
+		, const TSharedPtr<SWidget>& OldWidget, const FWidgetPath& WidgetPath, const TSharedPtr<SWidget>& NewWidget)
+		{
+			if (InFadeModeMenuWindow != nullptr && InFadeModeMenuWindow.IsValid())
+			{
+				if (TSharedPtr<SWindow> LockedFadeModeMenuWindow = InFadeModeMenuWindow.Pin())
+				{
+					if (OldWidget && !OldWidget->IsHovered() && LockedFadeModeMenuWindow == OldWidget)
+					{
+						FSlateApplication::Get().RequestDestroyWindow(LockedFadeModeMenuWindow.ToSharedRef());
+						FSlateApplication::Get().OnFocusChanging().Remove(*PopupHandlePtr);
+
+						if (InApplicationActivationStateHandle.IsValid())
+						{
+							FSlateApplication::Get().OnApplicationActivationStateChanged().Remove(InApplicationActivationStateHandle);
+						}
+					}
+				}
+			}
+		});
+
+	if (PopupHandle.IsValid())
+	{
+		FSlateApplication::Get().OnFocusChanging().Remove(PopupHandle);
+	}
+
+	PopupHandle = *PopupHandlePtr;
+
+	for (const TPair<EWaveEditorTransformationFadeMode, TSubclassOf<UTransformationFadeFunction>>& FadeOptionPair : UWaveformTransformationFade::FadeModeToFadeFunctionMap)
+	{
+		MenuContent->AddSlot()
+			.Padding(5)
+			.AutoHeight()
+			[
+				SNew(SButton)
+					.OnClicked_Lambda([InPopupHandle = PopupHandle, InApplicationActivationStateHandle = ApplicationActivationStateHandle, InOnFadeModeChanged = OnFadeModeChanged,
+						InWeakFadeModeMenuWindow = FadeModeMenuWindow, InWeakFade, InteractionTypeAndIndex = FadeInteractionTypeAndIndex,
+						FadeOptionPair, PropertyName = GetPropertyEditedByCurrentInteraction()]() -> FReply
+					{
+						if (TStrongObjectPtr<UWaveformTransformationFade> PinnedFade = InWeakFade.Pin())
+						{
+							const int32 Index = InteractionTypeAndIndex.Value;
+
+							if (Index < PinnedFade->GetFadeRegionsNum())
+							{
+								if (GEditor && GEditor->Trans)
+								{
+									GEditor->BeginTransaction(TEXT("PropertyEditor"), FText::Format(NSLOCTEXT("PropertyEditor", "SetCurveProperty", "Set Curve"), PropertyName), PinnedFade.Get());
+								}
+
+								PinnedFade->Modify();
+
+								if (InteractionTypeAndIndex.Key == FWaveformTransformationFadeRenderer::EFadeInteractionType::RightClickFadeIn)
+								{
+									if (PinnedFade->GetFadeRegions()[Index].FadeIn)
+									{
+										const float Duration = PinnedFade->GetFadeRegions()[Index].FadeIn->Duration;
+										const int64 Offset = PinnedFade->GetFadeRegions()[Index].FadeIn->FrameOffset;
+
+										PinnedFade->GetMutableFadeRegions()[Index].FadeIn = NewObject<UTransformationFadeFunction>(PinnedFade.Get(), FadeOptionPair.Value.Get(), NAME_None, RF_Transactional);
+										PinnedFade->GetFadeRegions()[Index].FadeIn->Duration = Duration;
+										PinnedFade->GetFadeRegions()[Index].FadeIn->FrameOffset = Offset;
+									}
+								}
+								else if (InteractionTypeAndIndex.Key == FWaveformTransformationFadeRenderer::EFadeInteractionType::RightClickFadeOut)
+								{
+									if (PinnedFade->GetFadeRegions()[Index].FadeOut)
+									{
+										const float Duration = PinnedFade->GetFadeRegions()[Index].FadeOut->Duration;
+										const int64 Offset = PinnedFade->GetFadeRegions()[Index].FadeOut->FrameOffset;
+
+										PinnedFade->GetMutableFadeRegions()[Index].FadeOut = NewObject<UTransformationFadeFunction>(PinnedFade.Get(), FadeOptionPair.Value.Get(), NAME_None, RF_Transactional);
+										PinnedFade->GetFadeRegions()[Index].FadeOut->Duration = Duration;
+										PinnedFade->GetFadeRegions()[Index].FadeOut->FrameOffset = Offset;
+									}
+								}
+
+								if (GEditor && GEditor->Trans)
+								{
+									GEditor->EndTransaction();
+								}
+
+								PinnedFade->PostEditChange();
+								InOnFadeModeChanged.ExecuteIfBound();
+							}
+						}
+
+						TSharedPtr<SWindow> LockedFadeModeMenuWindow = InWeakFadeModeMenuWindow.Pin();
+						if (LockedFadeModeMenuWindow)
+						{
+							FSlateApplication::Get().RequestDestroyWindow(LockedFadeModeMenuWindow.ToSharedRef());
+							
+							if (InPopupHandle.IsValid())
+							{
+								FSlateApplication::Get().OnFocusChanging().Remove(InPopupHandle);
+							}
+
+							if (InApplicationActivationStateHandle.IsValid())
+							{
+								FSlateApplication::Get().OnApplicationActivationStateChanged().Remove(InApplicationActivationStateHandle);
+							}
+						}
+
+						return FReply::Handled();
+					})
+					[
+						SNew(SHorizontalBox)
+						+ SHorizontalBox::Slot()
+						.AutoWidth()
+						.VAlign(VAlign_Center)
+						.HAlign(HAlign_Left)
+						[
+							SNew(SImage)
+							.Image(GetFadeModeIcon(FadeOptionPair.Key).GetIcon())
+						]
+						+ SHorizontalBox::Slot()
+						.VAlign(VAlign_Center)
+						.HAlign(HAlign_Center)
+						.Padding(5)
+						[
+							SNew(STextBlock)
+							.Text(FText::FromString(StaticEnum<EWaveEditorTransformationFadeMode>()->GetNameStringByValue(static_cast<int64>(FadeOptionPair.Key))))
+						]
+					]
+			];
+	}
+
+	FSlateApplication::Get().AddWindow(MenuWindow);
+
+	// The window does not have a valid position and size until after 1 tick
+	// If the new window is off screen, shift it back on screen
+	FTSTicker::GetCoreTicker().AddTicker(FTickerDelegate::CreateLambda([InFadeModeMenuWindow = FadeModeMenuWindow, LocalCursorPosition, LocalWindowMaxPosition](float DeltaTime)
+		{
+			if (TSharedPtr<SWindow> LockedFadeModeMenuWindow = InFadeModeMenuWindow.Pin())
+			{
+				const FVector2D WindowSize = LockedFadeModeMenuWindow->GetSizeInScreen();
+
+				if (WindowSize.X > 0.0 && LocalCursorPosition.X + WindowSize.X > LocalWindowMaxPosition.X)
+				{
+					check(WindowSize.Y > 0.0);
+
+					const double XPosition = FMath::Max(0.0, LocalCursorPosition.X - WindowSize.X);
+					LockedFadeModeMenuWindow->MoveWindowTo(FVector2D(XPosition, LocalCursorPosition.Y));
+				}
+			}
+
+			return false; // return false to only execute once
+		}));
+}
+
+const FSlateIcon FWaveformTransformationFadeRenderer::GetFadeModeIcon(const EWaveEditorTransformationFadeMode& FadeMode) const
+{
+	if (FadeInteractionTypeAndIndex.Key == FWaveformTransformationFadeRenderer::EFadeInteractionType::RightClickFadeIn)
+	{
+		switch (FadeMode)
+		{
+		case EWaveEditorTransformationFadeMode::Linear:
+			return FSlateIcon(FAudioWidgetsStyle::StyleName, "AudioWidgetsStyle.FadeInLinear");
+		case EWaveEditorTransformationFadeMode::Exponential:
+			return FSlateIcon(FAudioWidgetsStyle::StyleName, "AudioWidgetsStyle.FadeInExponential");
+		case EWaveEditorTransformationFadeMode::Logarithmic:
+			return FSlateIcon(FAudioWidgetsStyle::StyleName, "AudioWidgetsStyle.FadeInLogarithmic");
+		case EWaveEditorTransformationFadeMode::Sigmoid:
+			return FSlateIcon(FAudioWidgetsStyle::StyleName, "AudioWidgetsStyle.FadeInSigmoid");
+		default:
+			return FSlateIcon();
+		}
+	}
+	else if (FadeInteractionTypeAndIndex.Key == FWaveformTransformationFadeRenderer::EFadeInteractionType::RightClickFadeOut)
+	{
+		switch (FadeMode)
+		{
+		case EWaveEditorTransformationFadeMode::Linear:
+			return FSlateIcon(FAudioWidgetsStyle::StyleName, "AudioWidgetsStyle.FadeOutLinear");
+		case EWaveEditorTransformationFadeMode::Exponential:
+			return FSlateIcon(FAudioWidgetsStyle::StyleName, "AudioWidgetsStyle.FadeOutExponential");
+		case EWaveEditorTransformationFadeMode::Logarithmic:
+			return FSlateIcon(FAudioWidgetsStyle::StyleName, "AudioWidgetsStyle.FadeOutLogarithmic");
+		case EWaveEditorTransformationFadeMode::Sigmoid:
+			return FSlateIcon(FAudioWidgetsStyle::StyleName, "AudioWidgetsStyle.FadeOutSigmoid");
+		default:
+			return FSlateIcon();
+		}
+	}
+
+	return FSlateIcon();
+}
+
+TPair<FWaveformTransformationFadeRenderer::EFadeInteractionType, int32> FWaveformTransformationFadeRenderer::GetInteractionTypeAndIndexFromCursorPosition(const FVector2D& InLocalCursorPosition, const FKey MouseButton, const FGeometry& WidgetGeometry) const
+{
+	for (int32 Index = 0; Index < FadeRegions.Num(); Index++)
+	{
+		if (MouseButton == EKeys::LeftMouseButton)
+		{
+			if (FadeRegions[Index].FadeInStartInteractionXRange.Contains(InLocalCursorPosition.X) && FadeRegions[Index].FadeInStartInteractionYRange.Contains(InLocalCursorPosition.Y))
+			{
+				return TPair<EFadeInteractionType, int32>(EFadeInteractionType::ScrubbingFadeInStart, Index);
+			}
+
+			if (FadeRegions[Index].FadeInEndInteractionXRange.Contains(InLocalCursorPosition.X) && FadeRegions[Index].FadeInEndInteractionYRange.Contains(InLocalCursorPosition.Y))
+			{
+				return TPair<EFadeInteractionType, int32>(EFadeInteractionType::ScrubbingFadeInEnd, Index);
+			}
+
+			if (FadeRegions[Index].FadeOutStartInteractionXRange.Contains(InLocalCursorPosition.X) && FadeRegions[Index].FadeOutStartInteractionYRange.Contains(InLocalCursorPosition.Y))
+			{
+				return TPair<EFadeInteractionType, int32>(EFadeInteractionType::ScrubbingFadeOutStart, Index);
+			}
+
+			if (FadeRegions[Index].FadeOutEndInteractionXRange.Contains(InLocalCursorPosition.X) && FadeRegions[Index].FadeOutEndInteractionYRange.Contains(InLocalCursorPosition.Y))
+			{
+				return TPair<EFadeInteractionType, int32>(EFadeInteractionType::ScrubbingFadeOutEnd, Index);
+			}
+		}
+		else if (MouseButton == EKeys::RightMouseButton)
+		{
+			if (FadeRegions[Index].FadeInEndInteractionXRange.Contains(InLocalCursorPosition.X) && FadeRegions[Index].FadeInEndInteractionYRange.Contains(InLocalCursorPosition.Y))
+			{
+				return TPair<EFadeInteractionType, int32>(EFadeInteractionType::RightClickFadeIn, Index);
+			}
+
+			if (FadeRegions[Index].FadeOutStartInteractionXRange.Contains(InLocalCursorPosition.X) && FadeRegions[Index].FadeOutStartInteractionYRange.Contains(InLocalCursorPosition.Y))
+			{
+				return TPair<EFadeInteractionType, int32>(EFadeInteractionType::RightClickFadeOut, Index);
+			}
+		}
+	}
+
+	return TPair<EFadeInteractionType, int32>(EFadeInteractionType::None, 0);
+}
+
+FText FWaveformTransformationFadeRenderer::GetPropertyEditedByCurrentInteraction() const
+{
+
+	FText OutPropertyName = FText::FromName(NAME_None);
+
+	if (StrongFade->GetFadeRegions().IsValidIndex(FadeInteractionTypeAndIndex.Value))
+	{
+		check(StrongFade->GetFadeRegions()[FadeInteractionTypeAndIndex.Value].FadeIn);
+		check(StrongFade->GetFadeRegions()[FadeInteractionTypeAndIndex.Value].FadeOut);
+
+		switch (FadeInteractionTypeAndIndex.Key)
+		{
+		case FWaveformTransformationFadeRenderer::EFadeInteractionType::None:
+			break;
+		case FWaveformTransformationFadeRenderer::EFadeInteractionType::ScrubbingFadeInStart:
+			OutPropertyName = FText::Format(FText::FromString("{0}::{1}"),
+				FText::FromName(GET_MEMBER_NAME_CHECKED(FTransformationFadeFunctionData, FadeIn)),
+				FText::FromName(GET_MEMBER_NAME_CHECKED(UTransformationFadeFunction, FrameOffset)));
+
+			break;
+		case FWaveformTransformationFadeRenderer::EFadeInteractionType::ScrubbingFadeInEnd:
+			OutPropertyName = FText::Format(FText::FromString("{0}::{1}"),
+				FText::FromName(GET_MEMBER_NAME_CHECKED(FTransformationFadeFunctionData, FadeIn)),
+				FText::FromName(GET_MEMBER_NAME_CHECKED(UTransformationFadeFunction, Duration)));
+			
+			break;
+		case FWaveformTransformationFadeRenderer::EFadeInteractionType::ScrubbingFadeOutStart:
+			OutPropertyName = FText::Format(FText::FromString("{0}::{1}"),
+				FText::FromName(GET_MEMBER_NAME_CHECKED(FTransformationFadeFunctionData, FadeOut)),
+				FText::FromName(GET_MEMBER_NAME_CHECKED(UTransformationFadeFunction, Duration)));
+			
+			break;
+		case FWaveformTransformationFadeRenderer::EFadeInteractionType::ScrubbingFadeOutEnd:
+			OutPropertyName = FText::Format(FText::FromString("{0}::{1}"),
+				FText::FromName(GET_MEMBER_NAME_CHECKED(FTransformationFadeFunctionData, FadeOut)),
+				FText::FromName(GET_MEMBER_NAME_CHECKED(UTransformationFadeFunction, FrameOffset)));
+			
+			break;
+		default:
+			break;
+		}
+	}
+
+	return MoveTemp(OutPropertyName);
+}
\ No newline at end of file
diff --git a/Engine/Plugins/Editor/WaveformEditor/Source/WaveformTransformationsWidgets/Private/WaveformTransformationFadeRenderer.h b/Engine/Plugins/Editor/WaveformEditor/Source/WaveformTransformationsWidgets/Private/WaveformTransformationFadeRenderer.h
new file mode 100644
--- /dev/null
+++ b/Engine/Plugins/Editor/WaveformEditor/Source/WaveformTransformationsWidgets/Private/WaveformTransformationFadeRenderer.h
@@ -0,0 +1,95 @@
+// Copyright Epic Games, Inc. All Rights Reserved.
+
+#pragma once
+
+#include "WaveformTransformationRendererBase.h"
+#include "WaveformTransformationFade.h"
+#include "WaveformTransformationTrimFade.h"
+
+#include "Brushes/SlateRoundedBoxBrush.h"
+
+struct FFadeRegion
+{
+	uint32 FadeInStartX = 0;
+	uint32 FadeInEndX = 0;
+	uint32 FadeOutStartX = 0;
+	uint32 FadeOutEndX = 0;
+
+	TArray<FVector2D> FadeInCurvePoints;
+	TArray<FVector2D> FadeOutCurvePoints;
+
+	TRange<float> FadeInStartInteractionXRange = TRange<float>::Inclusive(0.f, 0.f);
+	TRange<float> FadeInEndInteractionXRange = TRange<float>::Inclusive(0.f, 0.f);
+	TRange<float> FadeInStartInteractionYRange = TRange<float>::Inclusive(0.f, 0.f);
+	TRange<float> FadeInEndInteractionYRange = TRange<float>::Inclusive(0.f, 0.f);
+	TRange<float> FadeOutStartInteractionXRange = TRange<float>::Inclusive(0.f, 0.f);
+	TRange<float> FadeOutEndInteractionXRange = TRange<float>::Inclusive(0.f, 0.f);
+	TRange<float> FadeOutStartInteractionYRange = TRange<float>::Inclusive(0.f, 0.f);
+	TRange<float> FadeOutEndInteractionYRange = TRange<float>::Inclusive(0.f, 0.f);
+};
+
+class FWaveformTransformationFadeRenderer : public FWaveformTransformationRendererBase
+{
+public:
+	FWaveformTransformationFadeRenderer() : RoundedBoxBrush(FLinearColor::Green, 2) {}
+	~FWaveformTransformationFadeRenderer();
+
+	virtual int32 OnPaint(const FPaintArgs& Args, const FGeometry& AllottedGeometry, const FSlateRect& MyCullingRect, FSlateWindowElementList& OutDrawElements, int32 LayerId, const FWidgetStyle& InWidgetStyle, bool bParentEnabled) const override;
+	virtual void Tick(const FGeometry& AllottedGeometry, const double InCurrentTime, const float InDeltaTime) override;
+	virtual FReply OnMouseButtonDown(SWidget& OwnerWidget, const FGeometry& MyGeometry, const FPointerEvent& MouseEvent) override;
+	virtual FReply OnMouseButtonUp(SWidget& OwnerWidget, const FGeometry& MyGeometry, const FPointerEvent& MouseEvent) override;
+	virtual FReply OnMouseMove(SWidget& OwnerWidget, const FGeometry& MyGeometry, const FPointerEvent& MouseEvent) override;
+	virtual FReply OnMouseWheel(SWidget& OwnerWidget, const FGeometry& MyGeometry, const FPointerEvent& MouseEvent) override;
+	virtual FCursorReply OnCursorQuery(const FGeometry& MyGeometry, const FPointerEvent& CursorEvent) const override;
+
+	virtual void SetWaveformTransformation(TObjectPtr<UWaveformTransformationBase> InTransformation) override;
+
+protected:
+	int32 DrawOffsetHandles(const FGeometry& AllottedGeometry, FSlateWindowElementList& OutDrawElements, int32 LayerId) const;
+	int32 DrawFadeCurves(const FGeometry& AllottedGeometry, FSlateWindowElementList& OutDrawElements, int32 LayerId) const;
+	void GenerateFadeCurves(const FGeometry& AllottedGeometry);
+	void UpdateInteractionRange(const FGeometry& AllottedGeometry);
+
+	bool IsCursorInFadeInteractionRange(const FVector2D& InLocalCursorPosition, const FGeometry& WidgetGeometry) const;
+
+	FText GetPropertyEditedByCurrentInteraction() const;
+	FVector2D GetLocalCursorPosition(const FPointerEvent& MouseEvent, const FGeometry& EventGeometry) const;
+	double ConvertXRatioToTime(const float InRatio) const;
+
+	TArray<FFadeRegion> FadeRegions;
+
+	double PixelsPerFrame = 0.0;
+	const float InteractionHandleSize = 30.0f;
+	FVector2D MousePosition;
+	int32 TransactionResult = INDEX_NONE;
+
+	enum class EFadeInteractionType : uint8
+	{
+		None = 0,
+		ScrubbingFadeInStart,
+		ScrubbingFadeInEnd, 
+		ScrubbingFadeOutStart,
+		ScrubbingFadeOutEnd,
+		RightClickFadeIn,
+		RightClickFadeOut
+	};
+
+	TPair<EFadeInteractionType, int32> FadeInteractionTypeAndIndex = TPair<EFadeInteractionType, int32>(EFadeInteractionType::None, 0);
+
+	TPair<EFadeInteractionType, int32> GetInteractionTypeAndIndexFromCursorPosition(const FVector2D& InLocalCursorPosition, const FKey MouseButton, const FGeometry& WidgetGeometry) const;
+	void SetPropertyValueDependingOnInteractionType(const FPointerEvent& MouseEvent, const FGeometry& WidgetGeometry);
+	void ShowSelectFadeModeMenuAtCursor(const FGeometry& MyGeometry, const FPointerEvent& MouseEvent);
+	const FSlateIcon GetFadeModeIcon(const EWaveEditorTransformationFadeMode& FadeMode) const;
+
+	const FSlateRoundedBoxBrush RoundedBoxBrush;
+
+	TWeakPtr<SWindow> FadeModeMenuWindow = nullptr;
+
+	TStrongObjectPtr<UWaveformTransformationFade> StrongFade = nullptr;
+
+	FDelegateHandle PopupHandle;
+	FDelegateHandle ApplicationActivationStateHandle;
+
+	DECLARE_DELEGATE(FOnFadeModeChanged);
+	FOnFadeModeChanged OnFadeModeChanged;
+};
diff --git a/Engine/Plugins/Editor/WaveformEditor/Source/WaveformTransformationsWidgets/Private/WaveformTransformationTrimFadeRenderer.cpp b/Engine/Plugins/Editor/WaveformEditor/Source/WaveformTransformationsWidgets/Private/WaveformTransformationTrimFadeRenderer.cpp
--- a/Engine/Plugins/Editor/WaveformEditor/Source/WaveformTransformationsWidgets/Private/WaveformTransformationTrimFadeRenderer.cpp
+++ b/Engine/Plugins/Editor/WaveformEditor/Source/WaveformTransformationsWidgets/Private/WaveformTransformationTrimFadeRenderer.cpp
@@ -2,403 +2,48 @@
 
 #include "WaveformTransformationTrimFadeRenderer.h"
 
-#include "AudioWidgetsStyle.h"
+#include "Blueprint/WidgetLayoutLibrary.h"
 #include "Editor.h"
 #include "Framework/Application/SlateApplication.h"
-#include "Layout/Geometry.h"
-#include "PropertyHandle.h"
-#include "Rendering/DrawElements.h"
-#include "Sound/SoundWave.h"
 #include "WaveformTransformationLog.h"
 #include "Widgets/Images/SImage.h"
 #include "Widgets/Input/SButton.h"
 #include "Widgets/SWidget.h"
 #include "Widgets/Text/STextBlock.h"
-#include <Blueprint/WidgetLayoutLibrary.h>
 
-FWaveformTransformationTrimFadeRenderer::~FWaveformTransformationTrimFadeRenderer()
-{
-	if (FSlateApplication::IsInitialized())
-	{
-		FSlateApplication::Get().OnFocusChanging().Remove(PopupHandle);
-		FSlateApplication::Get().OnApplicationActivationStateChanged().Remove(ApplicationActivationStateHandle);
-	}
-}
-
-int32 FWaveformTransformationTrimFadeRenderer::OnPaint(const FPaintArgs& Args, const FGeometry& AllottedGeometry, const FSlateRect& MyCullingRect, FSlateWindowElementList& OutDrawElements, int32 LayerId, const FWidgetStyle& InWidgetStyle, bool bParentEnabled) const
-{
-	LayerId = DrawTrimHandles(AllottedGeometry, OutDrawElements, LayerId);
-	LayerId = DrawFadeCurves(AllottedGeometry, OutDrawElements, LayerId);
-
-	return LayerId;
-}
-
-int32 FWaveformTransformationTrimFadeRenderer::DrawTrimHandles(const FGeometry& AllottedGeometry, FSlateWindowElementList& OutDrawElements, int32 LayerId) const
-{
-	const bool bRenderLowerBound = StartTimeHandleX >= 0.f;
-	const bool bRenderUpperBound = EndTimeHandleX <= AllottedGeometry.Size.X;
-
-	const float HandleWidth = InteractionHandleSize;
-
-	FPaintGeometry HandleGeometry = AllottedGeometry.ToPaintGeometry(FVector2f(HandleWidth, HandleWidth), FSlateLayoutTransform(FVector2f(StartTimeHandleX, AllottedGeometry.Size.Y - HandleWidth)));
-
-	FSlateDrawElement::MakeBox(
-		OutDrawElements,
-		++LayerId,
-		HandleGeometry,
-		&RoundedBoxBrush,
-		ESlateDrawEffect::None,
-		FLinearColor::Green
-	);
-
-	HandleGeometry = AllottedGeometry.ToPaintGeometry(FVector2f(HandleWidth, HandleWidth), FSlateLayoutTransform(FVector2f(EndTimeHandleX - HandleWidth, AllottedGeometry.Size.Y - HandleWidth)));
-
-	FSlateDrawElement::MakeBox(
-		OutDrawElements,
-		++LayerId,
-		HandleGeometry,
-		&RoundedBoxBrush,
-		ESlateDrawEffect::None,
-		FLinearColor::Red
-	);
-
-	return LayerId;
-}
-
-int32 FWaveformTransformationTrimFadeRenderer::DrawFadeCurves(const FGeometry& AllottedGeometry, FSlateWindowElementList& OutDrawElements, int32 LayerId) const
-{
-	const FLinearColor FadeColor = FLinearColor::Gray;
-	const float HandleWidth = InteractionHandleSize;
-	const float HandleHeight = InteractionHandleSize;
-	const float HandleOffset = 5.0f;
-
-	const FLinearColor HandleColor = FLinearColor::White;
-	const FLinearColor HandleSelectedColor = FLinearColor::Yellow;
-
-	const bool IsFadeInSelected = IsCursorInFadeInInteractionRange(MousePosition, AllottedGeometry);
-	const bool IsFadeOutSelected = IsCursorInFadeOutInteractionRange(MousePosition, AllottedGeometry);
-	
-	if (FadeInCurvePoints.Num() > 0)
-	{
-		FSlateDrawElement::MakeLines(
-			OutDrawElements,
-			++LayerId,
-			AllottedGeometry.ToPaintGeometry(),
-			FadeInCurvePoints,
-			ESlateDrawEffect::None,
-			FadeColor
-		);
-
-		const float	HandleStart = FadeInCurvePoints.Last(0).X;
-		const FPaintGeometry HandleGeometry = AllottedGeometry.ToPaintGeometry(FVector2f(HandleWidth, HandleHeight), FSlateLayoutTransform(FVector2f(HandleStart - HandleOffset, 0)));
-		const FSlateBrush* HandleBrush = FAppStyle::GetBrush("WorldBrowser.DirectionXPositive");
-	
-		FSlateDrawElement::MakeBox(
-			OutDrawElements,
-			LayerId,
-			HandleGeometry,
-			HandleBrush,
-			ESlateDrawEffect::DisabledEffect,
-			IsFadeInSelected ? HandleSelectedColor : HandleColor
-		);
-	}
-	else
-	{
-		FPaintGeometry HandleGeometry = AllottedGeometry.ToPaintGeometry(FVector2f(HandleWidth, HandleHeight), FSlateLayoutTransform(FVector2f(StartTimeHandleX - HandleOffset, 0)));
-		const FSlateBrush* HandleBrush = FAppStyle::GetBrush("WorldBrowser.DirectionXPositive");
-
-		FSlateDrawElement::MakeBox(
-			OutDrawElements,
-			LayerId,
-			HandleGeometry,
-			HandleBrush,
-			ESlateDrawEffect::DisabledEffect,
-			IsFadeInSelected ? HandleSelectedColor : HandleColor
-		);
-	}
-
-	if (FadeOutCurvePoints.Num() > 0)
-	{
-		FSlateDrawElement::MakeLines(
-			OutDrawElements,
-			++LayerId,
-			AllottedGeometry.ToPaintGeometry(),
-			FadeOutCurvePoints,
-			ESlateDrawEffect::None,
-			FadeColor
-		);
-
-		const float	HandleStart = FadeOutCurvePoints[0].X - HandleWidth;
-		const FPaintGeometry HandleGeometry = AllottedGeometry.ToPaintGeometry(FVector2f(HandleWidth, HandleHeight), FSlateLayoutTransform(FVector2f(HandleStart + HandleOffset, 0)));
-		const FSlateBrush* HandleBrush = FAppStyle::GetBrush("WorldBrowser.DirectionXNegative");
-
-		FSlateDrawElement::MakeBox(
-			OutDrawElements,
-			LayerId,
-			HandleGeometry,
-			HandleBrush,
-			ESlateDrawEffect::DisabledEffect,
-			IsFadeOutSelected ? HandleSelectedColor : HandleColor
-		);
-	}
-	else
-	{
-		FPaintGeometry HandleGeometry = AllottedGeometry.ToPaintGeometry(FVector2f(HandleWidth, HandleHeight), FSlateLayoutTransform(FVector2f(EndTimeHandleX - HandleWidth + HandleOffset, 0)));
-		const FSlateBrush* HandleBrush = FAppStyle::GetBrush("WorldBrowser.DirectionXNegative");
-
-		FSlateDrawElement::MakeBox(
-			OutDrawElements,
-			LayerId,
-			HandleGeometry,
-			HandleBrush,
-			ESlateDrawEffect::DisabledEffect,
-			IsFadeOutSelected ? HandleSelectedColor : HandleColor
-		);
-	}
-
-	return LayerId;
-}
-
-void FWaveformTransformationTrimFadeRenderer::GenerateFadeCurves(const FGeometry& AllottedGeometry)
+void FWaveformTransformationTrimFadeRenderer::Tick(const FGeometry& AllottedGeometry, const double InCurrentTime, const float InDeltaTime)
 {
+	check(TransformationWaveInfo.NumChannels > 0);
+	check(StrongFade);
 	check(StrongTrimFade);
 
-	check(StartTimeHandleX >= 0);
-	check(EndTimeHandleX >= 0);
-	FadeInStartX = FMath::RoundToInt32(StartTimeHandleX);
-	FadeOutEndX = FMath::RoundToInt32(EndTimeHandleX);
-
-	if (StrongTrimFade->FadeFunctions.FadeIn && StrongTrimFade->FadeFunctions.FadeIn->Duration > 0.f && PixelsPerFrame > 0.0)
-	{
-		const float FadeInTime = StrongTrimFade->FadeFunctions.FadeIn->Duration;
-		const float FadeInFrames = FadeInTime * TransformationWaveInfo.SampleRate;
-		const float FadeInPixelLength = FadeInFrames * PixelsPerFrame;
-		check(FadeInPixelLength > 0.f);
-		check(EndTimeHandleX >= 0);
-		
-		FadeInEndX = FMath::RoundToInt32(FMath::Clamp(StartTimeHandleX + FadeInPixelLength, StartTimeHandleX, EndTimeHandleX));
-
-		const uint32 DisplayedFadeInPixelLength = FadeInStartX < FadeInEndX ? FadeInEndX - FadeInStartX : 0;
-		
-		FadeInCurvePoints.SetNumUninitialized(DisplayedFadeInPixelLength);
-
-		if (DisplayedFadeInPixelLength > 0)
-		{
-			for (uint32 Pixel = 0; Pixel < DisplayedFadeInPixelLength; ++Pixel)
-			{
-				const double FadeFraction = static_cast<float>(Pixel) / FadeInPixelLength;
-				const double CurveFunction = StrongTrimFade->FadeFunctions.FadeIn->GetFadeInCurveValue(FadeFraction);
-
-				const bool bNeedsFlickerPrevention = (Pixel == DisplayedFadeInPixelLength - 1) && DisplayedFadeInPixelLength != 1;
-
-				// Prevent fade line from flickering if the math doesn't work out to the last Y value being 0 as the position is changing
-				// This is due to pixel rounding errors. The slight distortion caused by this correction isn't noticeable
-				const double CurveValue = bNeedsFlickerPrevention ? 0.f : 1.f - CurveFunction;
-
-				const uint32 XCoordinate = Pixel + FadeInStartX;
-				FadeInCurvePoints[Pixel] = FVector2D(XCoordinate, CurveValue * AllottedGeometry.Size.Y);
-			}
-
-			check(FadeInCurvePoints.Num() > 0);
-			check(FadeInCurvePoints[0].Y == AllottedGeometry.Size.Y);
-		}
-	}
-	else
-	{
-		FadeInEndX = FMath::RoundToInt32(FMath::Clamp(StartTimeHandleX, StartTimeHandleX, EndTimeHandleX));
-		FadeInCurvePoints.SetNumUninitialized(0);
-	}
-
-	if(StrongTrimFade->FadeFunctions.FadeOut && StrongTrimFade->FadeFunctions.FadeOut->Duration > 0.f && PixelsPerFrame > 0.0)
-	{
-		const float FadeOutTime = StrongTrimFade->FadeFunctions.FadeOut->Duration;
-		const float FadeOutFrames = FadeOutTime * TransformationWaveInfo.SampleRate;
-		const float FadeOutPixelLength = FadeOutFrames * PixelsPerFrame;
-		check(FadeOutPixelLength > 0.f);
-
-		check(StartTimeHandleX >= 0);
-		FadeOutStartX = FMath::RoundToInt32(FMath::Clamp(EndTimeHandleX - FadeOutPixelLength, StartTimeHandleX, EndTimeHandleX));
-
-		const uint32 DisplayedFadeOutPixelLength = 
-			(FadeOutStartX < FadeOutEndX) ? (FadeOutEndX - FadeOutStartX) : 0;
-
-		FadeOutCurvePoints.SetNumUninitialized(DisplayedFadeOutPixelLength);
-		
-		if (DisplayedFadeOutPixelLength > 0)
-		{
-			for (uint32 Pixel = 0; Pixel < DisplayedFadeOutPixelLength; ++Pixel)
-			{
-				const double FadeFraction = static_cast<float>(Pixel) / FadeOutPixelLength;
-				const double CurveFunction = StrongTrimFade->FadeFunctions.FadeOut->GetFadeOutCurveValue(FadeFraction);
-
-				const bool bNeedsFlickerPrevention = (Pixel == DisplayedFadeOutPixelLength - 1) && DisplayedFadeOutPixelLength != 1;
-
-				// Prevent fade line from flickering if the math doesn't work out to the last Y value being equal to AllottedGeometry.Size.Y as the 
-				// position is changing. This is due to pixel rounding errors. The slight distortion caused by this correction isn't noticeable
-				const double CurveValue = bNeedsFlickerPrevention ? 1.f : 1.f - CurveFunction;
-
-				const uint32 XCoordinate = Pixel + FadeOutStartX;
-				FadeOutCurvePoints[Pixel] = FVector2D(XCoordinate, CurveValue * AllottedGeometry.Size.Y);
-			}
-
-			check(FadeOutCurvePoints.Num() > 0);
-			check(FadeOutCurvePoints[0].Y == 0.f);
-		}
-	}
-	else
-	{
-		FadeOutStartX = FMath::RoundToInt32(FMath::Clamp(EndTimeHandleX, StartTimeHandleX, EndTimeHandleX));
-		FadeOutCurvePoints.SetNumUninitialized(0);
-	}
-}
-
-FCursorReply FWaveformTransformationTrimFadeRenderer::OnCursorQuery(const FGeometry& MyGeometry, const FPointerEvent& CursorEvent) const
-{
-	const FVector2D LocalCursorPosition = GetLocalCursorPosition(CursorEvent, MyGeometry);
-
-	if (TrimFadeInteractionType == ETrimFadeInteractionType::ScrubbingFadeIn || TrimFadeInteractionType == ETrimFadeInteractionType::ScrubbingFadeOut)
-	{
-		return FCursorReply::Cursor(EMouseCursor::GrabHandClosed);
-	}
-
-	if (IsCursorInFadeInInteractionRange(LocalCursorPosition, MyGeometry) || IsCursorInFadeOutInteractionRange(LocalCursorPosition, MyGeometry))
-	{
-		return FCursorReply::Cursor(EMouseCursor::GrabHand);
-	}
-
-	if ((TrimFadeInteractionType == ETrimFadeInteractionType::ScrubbingLeftHandle || TrimFadeInteractionType == ETrimFadeInteractionType::ScrubbingRightHandle) 
-		|| (EndTimeInteractionXRange.Contains(LocalCursorPosition.X) && EndTimeInteractionYRange.Contains(LocalCursorPosition.Y))
-		|| (StartTimeInteractionXRange.Contains(LocalCursorPosition.X) && StartTimeInteractionYRange.Contains(LocalCursorPosition.Y)))
-	{
-		return FCursorReply::Cursor(EMouseCursor::ResizeLeftRight);
-	}
-
-	return FCursorReply::Unhandled();
-}
-
-void FWaveformTransformationTrimFadeRenderer::SetWaveformTransformation(TObjectPtr<UWaveformTransformationBase> InTransformation)
-{
-	FWaveformTransformationRendererBase::SetWaveformTransformation(InTransformation);
-	StrongTrimFade = TStrongObjectPtr<UWaveformTransformationTrimFade>(CastChecked<UWaveformTransformationTrimFade>(InTransformation.Get()));
-}
+	const float NumFrames = TransformationWaveInfo.TotalNumSamples / TransformationWaveInfo.NumChannels;
 
-FReply FWaveformTransformationTrimFadeRenderer::OnMouseWheel(SWidget& OwnerWidget, const FGeometry& MyGeometry, const FPointerEvent& MouseEvent)
-{
-	check(StrongTrimFade);
+	check(NumFrames > 0);
+	PixelsPerFrame = FMath::Max(AllottedGeometry.GetLocalSize().X / NumFrames, 0.0);
 
-	// Prevent stomping on other mouse transactions
-	if (TrimFadeInteractionType != ETrimFadeInteractionType::None)
+	const int32 NumFadeRegions = StrongFade->GetFadeRegionsNum();
+	if (NumFadeRegions != FadeRegions.Num())
 	{
-		return FReply::Handled();
+		FadeRegions.SetNum(NumFadeRegions);
 	}
 
-	const FVector2D LocalCursorPosition = GetLocalCursorPosition(MouseEvent, MyGeometry);
+	check(NumFadeRegions == FadeRegions.Num());
 
-	if (IsCursorInFadeInInteractionRange(LocalCursorPosition, MyGeometry) && StrongTrimFade->FadeFunctions.FadeIn && StrongTrimFade->FadeFunctions.FadeIn.IsA<UFadeCurveFunction>())
+	for (int32 Index = 0; Index < NumFadeRegions; Index++)
 	{
-		const float FadeCurveDelta = MouseEvent.GetWheelDelta() * MouseWheelStep;
-		FText OutPropertyName = FText::FromName(NAME_None);
+		check(StrongFade->GetFadeRegions()[Index].FadeIn);
+		check(StrongFade->GetFadeRegions()[Index].FadeOut);
 
-		TObjectPtr<UFadeCurveFunction> FadeIn = CastChecked<UFadeCurveFunction>(StrongTrimFade->FadeFunctions.FadeIn);
-		OutPropertyName = FadeIn->GetFadeCurvePropertyName();
+		const double FirstFrame = FMath::Clamp((StrongTrimFade->StartTime * TransformationWaveInfo.SampleRate), 0.f, NumFrames);
+		const double EndFrame = FMath::Clamp((StrongTrimFade->EndTime * TransformationWaveInfo.SampleRate), FirstFrame, NumFrames);
 
-		if (!GIsTransacting)
-		{
-			TransactionResult =
-				BeginTransaction(TEXT("PropertyEditor"), FText::Format(NSLOCTEXT("PropertyEditor", "SetNumericPropertyTransaction", "Edit {0}"), OutPropertyName), StrongTrimFade.Get());
-
-			if (TransactionResult == INDEX_NONE)
-			{
-				UE_LOG(LogWaveformTransformation, Warning, TEXT("Begin TrimFadeTransformation OnMouseWheel Transaction Failed"));
-			}
-		}
-
-		StrongTrimFade->Modify();
-		StrongTrimFade->FadeFunctions.FadeIn->Modify();
-
-		const float FadeCurveValue = FadeIn->GetFadeCurve() + FadeCurveDelta;
-		FadeIn->SetFadeCurve(FadeCurveValue);
-
-		if (TransactionResult != INDEX_NONE)
-		{
-			EndTransaction();
-			TransactionResult = INDEX_NONE;
-		}
-		StrongTrimFade->OnTransformationChanged.ExecuteIfBound(true);
-
-		return FReply::Handled();
+		FadeRegions[Index].FadeInStartX = static_cast<uint32>(FMath::Clamp((FirstFrame * PixelsPerFrame), 0, UINT32_MAX));
+		FadeRegions[Index].FadeOutEndX = static_cast<uint32>(FMath::Clamp((EndFrame * PixelsPerFrame), 0, UINT32_MAX));
+		FadeRegions[Index].FadeInEndX = static_cast<uint32>(FMath::Clamp(StrongFade->GetFadeRegions()[Index].FadeIn->Duration * TransformationWaveInfo.SampleRate * PixelsPerFrame + FadeRegions[Index].FadeInStartX, 0, UINT32_MAX));
+		FadeRegions[Index].FadeOutStartX = static_cast<uint32>(FMath::Clamp(FadeRegions[Index].FadeOutEndX - (StrongFade->GetFadeRegions()[Index].FadeOut->Duration * TransformationWaveInfo.SampleRate * PixelsPerFrame), 0, UINT32_MAX));
 	}
 
-	if (IsCursorInFadeOutInteractionRange(LocalCursorPosition, MyGeometry) && StrongTrimFade->FadeFunctions.FadeOut && StrongTrimFade->FadeFunctions.FadeOut.IsA<UFadeCurveFunction>())
-	{
-		const float FadeCurveDelta = MouseEvent.GetWheelDelta() * MouseWheelStep;
-		FText OutPropertyName = FText::FromName(NAME_None);
-
-		TObjectPtr<UFadeCurveFunction> FadeOut = CastChecked<UFadeCurveFunction>(StrongTrimFade->FadeFunctions.FadeOut);
-		OutPropertyName = FadeOut->GetFadeCurvePropertyName();
-
-		if (!GIsTransacting)
-		{
-			TransactionResult =
-				BeginTransaction(TEXT("PropertyEditor"), FText::Format(NSLOCTEXT("PropertyEditor", "SetNumericPropertyTransaction", "Edit {0}"), OutPropertyName),
-					StrongTrimFade.Get());
-			
-			if (TransactionResult == INDEX_NONE)
-			{
-				UE_LOG(LogWaveformTransformation, Warning, TEXT("Begin TrimFadeTransformation OnMouseWheel Transaction Failed"));
-			}
-		}
-
-		StrongTrimFade->Modify();
-		StrongTrimFade->FadeFunctions.FadeOut->Modify();
-
-		const float FadeCurveValue = FadeOut->GetFadeCurve() + FadeCurveDelta;
-		FadeOut->SetFadeCurve(FadeCurveValue);
-
-		if (TransactionResult != INDEX_NONE)
-		{
-			EndTransaction();
-			TransactionResult = INDEX_NONE;
-		}
-		StrongTrimFade->OnTransformationChanged.ExecuteIfBound(true);
-
-		return FReply::Handled();
-	}
-
-	return FReply::Unhandled();
-}
-
-FVector2D FWaveformTransformationTrimFadeRenderer::GetLocalCursorPosition(const FPointerEvent& MouseEvent, const FGeometry& EventGeometry) const
-{
-	const FVector2D ScreenSpacePosition = MouseEvent.GetScreenSpacePosition();
-	return  EventGeometry.AbsoluteToLocal(ScreenSpacePosition);
-}
-
-double FWaveformTransformationTrimFadeRenderer::ConvertXRatioToTime(const float InRatio) const
-{
-	check(TransformationWaveInfo.NumChannels > 0);
-	check(TransformationWaveInfo.SampleRate > 0);
-
-	const float NumFrames = TransformationWaveInfo.TotalNumSamples / TransformationWaveInfo.NumChannels;
-	const float FrameSelected = NumFrames * InRatio;
-	return FrameSelected / TransformationWaveInfo.SampleRate;
-}
-
-void FWaveformTransformationTrimFadeRenderer::Tick(const FGeometry& AllottedGeometry, const double InCurrentTime, const float InDeltaTime)
-{ 
-	check(TransformationWaveInfo.NumChannels > 0);
-	check(StrongTrimFade);
-
-	const float NumFrames = TransformationWaveInfo.TotalNumSamples / TransformationWaveInfo.NumChannels;
-	const double FirstFrame = FMath::Clamp((StrongTrimFade->StartTime * TransformationWaveInfo.SampleRate) , 0.f, NumFrames);
-	const double EndFrame = FMath::Clamp((StrongTrimFade->EndTime * TransformationWaveInfo.SampleRate), FirstFrame, NumFrames);
-
-	check(NumFrames > 0);
-	PixelsPerFrame = FMath::Max(AllottedGeometry.GetLocalSize().X / NumFrames, 0.0);
-
-	StartTimeHandleX = FirstFrame * PixelsPerFrame;
-	EndTimeHandleX = EndFrame * PixelsPerFrame;
-
 	FVector2D MouseAbsolutePosition = UWidgetLayoutLibrary::GetMousePositionOnPlatform();
 	MousePosition = AllottedGeometry.AbsoluteToLocal(MouseAbsolutePosition);
 
@@ -406,83 +51,24 @@ void FWaveformTransformationTrimFadeRenderer::Tick(const FGeometry& AllottedGeom
 	UpdateInteractionRange(AllottedGeometry);
 }
 
-void FWaveformTransformationTrimFadeRenderer::UpdateInteractionRange(const FGeometry& AllottedGeometry)
-{
-	StartTimeInteractionXRange.SetLowerBoundValue(StartTimeHandleX - InteractionHandleSize);
-	StartTimeInteractionXRange.SetUpperBoundValue(StartTimeHandleX + InteractionHandleSize);
-	StartTimeInteractionYRange.SetLowerBoundValue(AllottedGeometry.GetLocalSize().Y - InteractionHandleSize);
-	StartTimeInteractionYRange.SetUpperBoundValue(AllottedGeometry.GetLocalSize().Y);
-
-	EndTimeInteractionXRange.SetLowerBoundValue(EndTimeHandleX - InteractionHandleSize);
-	EndTimeInteractionXRange.SetUpperBoundValue(EndTimeHandleX + InteractionHandleSize);
-	EndTimeInteractionYRange.SetLowerBoundValue(AllottedGeometry.GetLocalSize().Y - InteractionHandleSize);
-	EndTimeInteractionYRange.SetUpperBoundValue(AllottedGeometry.GetLocalSize().Y);
-
-	FadeInInteractionXRange.SetLowerBoundValue(static_cast<float>(FadeInEndX) - InteractionHandleSize);
-	FadeInInteractionXRange.SetUpperBoundValue(static_cast<float>(FadeInEndX) + InteractionHandleSize);
-	FadeOutInteractionXRange.SetLowerBoundValue(static_cast<float>(FadeOutStartX) - InteractionHandleSize);
-	FadeOutInteractionXRange.SetUpperBoundValue(static_cast<float>(FadeOutStartX) + InteractionHandleSize);
-}
-
 FReply FWaveformTransformationTrimFadeRenderer::OnMouseButtonDown(SWidget& OwnerWidget, const FGeometry& MyGeometry, const FPointerEvent& MouseEvent)
 {
 	check(StrongTrimFade);
-	const FVector2D LocalCursorPosition = GetLocalCursorPosition(MouseEvent, MyGeometry);
-
-	const FKey MouseButton = MouseEvent.GetEffectingButton();
-
-	TrimFadeInteractionType = GetInteractionTypeFromCursorPosition(LocalCursorPosition, MouseButton, MyGeometry);
+	FReply Reply = FWaveformTransformationFadeRenderer::OnMouseButtonDown(OwnerWidget, MyGeometry, MouseEvent);
 
-	if (TrimFadeInteractionType != ETrimFadeInteractionType::None)
+	if (FadeInteractionTypeAndIndex.Key != EFadeInteractionType::None)
 	{
-		if (!GIsTransacting)
-		{
-			TransactionResult =
-				BeginTransaction(TEXT("PropertyEditor"), FText::Format(NSLOCTEXT("PropertyEditor", "SetNumericPropertyTransaction", "Edit {0}"), GetPropertyEditedByCurrentInteraction()), StrongTrimFade.Get());
-
-			if (TransactionResult == INDEX_NONE)
-			{
-				UE_LOG(LogWaveformTransformation, Warning, TEXT("Begin TrimFadeTransformation OnMouseButtonDown Transaction Failed"));
-			}
-		}
-
-		if (TrimFadeInteractionType == ETrimFadeInteractionType::ScrubbingFadeIn && StrongTrimFade->FadeFunctions.FadeIn)
-		{
-			StrongTrimFade->FadeFunctions.FadeIn->Modify();
-		}
-		else if (TrimFadeInteractionType == ETrimFadeInteractionType::ScrubbingFadeOut && StrongTrimFade->FadeFunctions.FadeOut)
-		{
-			StrongTrimFade->FadeFunctions.FadeOut->Modify();
-		}
-		
 		StrongTrimFade->Modify();
-
-		return FReply::Handled().CaptureMouse(OwnerWidget.AsShared()).PreventThrottling();
 	}
 
-	return FReply::Unhandled();
-}
-
-
-FReply FWaveformTransformationTrimFadeRenderer::OnMouseMove(SWidget& OwnerWidget, const FGeometry& MyGeometry, const FPointerEvent& MouseEvent)
-{
-	check(StrongTrimFade);
-	if (MouseEvent.IsMouseButtonDown(EKeys::LeftMouseButton) && TrimFadeInteractionType != ETrimFadeInteractionType::None)
-	{
-		SetPropertyValueDependingOnInteractionType(MouseEvent, MyGeometry);
-
-		StrongTrimFade->PostEditChange();
-
-		return FReply::Handled().CaptureMouse(OwnerWidget.AsShared());
-	}
-
-	return FReply::Unhandled();
+	return Reply;
 }
 
 FReply FWaveformTransformationTrimFadeRenderer::OnMouseButtonUp(SWidget& OwnerWidget, const FGeometry& MyGeometry, const FPointerEvent& MouseEvent)
 {
 	check(StrongTrimFade);
-	if (TrimFadeInteractionType != ETrimFadeInteractionType::None)
+
+	if (FadeInteractionTypeAndIndex.Key != EFadeInteractionType::None)
 	{
 		SetPropertyValueDependingOnInteractionType(MouseEvent, MyGeometry);
 
@@ -494,377 +80,104 @@ FReply FWaveformTransformationTrimFadeRenderer::OnMouseButtonUp(SWidget& OwnerWi
 
 		StrongTrimFade->PostEditChange();
 
-		StrongTrimFade->OnTransformationChanged.ExecuteIfBound(true);
+		FadeInteractionTypeAndIndex.Key = EFadeInteractionType::None;
+		FadeInteractionTypeAndIndex.Value = 0;
 
-		TrimFadeInteractionType = ETrimFadeInteractionType::None;
-		
 		return FReply::Handled().ReleaseMouseCapture();
 	}
 
 	return FReply::Unhandled();
 }
 
-void FWaveformTransformationTrimFadeRenderer::SetPropertyValueDependingOnInteractionType(const FPointerEvent& MouseEvent, const FGeometry& WidgetGeometry)
+FReply FWaveformTransformationTrimFadeRenderer::OnMouseMove(SWidget& OwnerWidget, const FGeometry& MyGeometry, const FPointerEvent& MouseEvent)
 {
 	check(StrongTrimFade);
-	check(WidgetGeometry.GetLocalSize().X > 0.f);
-	check(WidgetGeometry.GetLocalSize().X > 0.f);
-	const FVector2D LocalCursorPosition = GetLocalCursorPosition(MouseEvent, WidgetGeometry);
-	const float LocalCursorXRatio = FMath::Clamp(LocalCursorPosition.X / WidgetGeometry.GetLocalSize().X, 0.f, 1.f);
-	const double SelectedTime = ConvertXRatioToTime(LocalCursorXRatio);
-
-	switch (TrimFadeInteractionType)
+	if (MouseEvent.IsMouseButtonDown(EKeys::LeftMouseButton) && FadeInteractionTypeAndIndex.Key != EFadeInteractionType::None)
 	{
-	case FWaveformTransformationTrimFadeRenderer::ETrimFadeInteractionType::None:
-		break;
-	case FWaveformTransformationTrimFadeRenderer::ETrimFadeInteractionType::ScrubbingLeftHandle:
-		StrongTrimFade->StartTime = FMath::Min(SelectedTime, StrongTrimFade->EndTime);
-		break;
-	case FWaveformTransformationTrimFadeRenderer::ETrimFadeInteractionType::ScrubbingRightHandle:
-		StrongTrimFade->EndTime = FMath::Max(SelectedTime, StrongTrimFade->StartTime);
-		break;
-	case FWaveformTransformationTrimFadeRenderer::ETrimFadeInteractionType::ScrubbingFadeIn:
-		{
-			float StartFadeTimeValue = 0;
-			if (StrongTrimFade->FadeFunctions.FadeIn)
-			{
-				StartFadeTimeValue = FMath::Clamp(SelectedTime - StrongTrimFade->StartTime, 0.f, TNumericLimits<float>().Max());
-				StrongTrimFade->FadeFunctions.FadeIn->Duration = StartFadeTimeValue;
-			}
-		}
-		break;
-	case FWaveformTransformationTrimFadeRenderer::ETrimFadeInteractionType::ScrubbingFadeOut:
-		{	
-			float EndFadeTimeValue = 0;
-			if (StrongTrimFade->FadeFunctions.FadeOut)
-			{
-				EndFadeTimeValue = FMath::Clamp(StrongTrimFade->EndTime - SelectedTime, 0.f, TNumericLimits<float>().Max());
-				StrongTrimFade->FadeFunctions.FadeOut->Duration = EndFadeTimeValue;
-			}
-		}
-		break;
-	case FWaveformTransformationTrimFadeRenderer::ETrimFadeInteractionType::RightClickFadeIn:
-	case FWaveformTransformationTrimFadeRenderer::ETrimFadeInteractionType::RightClickFadeOut:
-	{
-		ShowSelectFadeModeMenuAtCursor(WidgetGeometry, MouseEvent);
-	}
-		break;
-	default:
-		break;
-	}
-}
+		SetPropertyValueDependingOnInteractionType(MouseEvent, MyGeometry);
 
-void FWaveformTransformationTrimFadeRenderer::ShowSelectFadeModeMenuAtCursor(const FGeometry& MyGeometry, const FPointerEvent& MouseEvent)
-{
-	if (!FSlateApplication::IsInitialized())
-	{
-		UE_LOG(LogWaveformTransformation, Warning, TEXT("FSlateApplication not initialized in WaveformTransformationTrimFadeRenderer."));
+		StrongTrimFade->PostEditChange();
 
-		return;
+		return FReply::Handled().CaptureMouse(OwnerWidget.AsShared());
 	}
 
-	const FVector2D LocalWindowMaxPosition = MyGeometry.GetAbsolutePosition() + MyGeometry.GetAbsoluteSize();
-	const FVector2D LocalCursorPosition = MouseEvent.GetScreenSpacePosition();
-
-	if (TSharedPtr<SWindow> LockedFadeModeMenuWindow = FadeModeMenuWindow.Pin())
-	{
-		FSlateApplication::Get().RequestDestroyWindow(LockedFadeModeMenuWindow.ToSharedRef());
-		FSlateApplication::Get().OnFocusChanging().Remove(PopupHandle); 
-		FSlateApplication::Get().OnApplicationActivationStateChanged().Remove(ApplicationActivationStateHandle);
-	}
+	return FReply::Unhandled();
+}
 
-	TSharedRef<SVerticalBox> MenuContent = SNew(SVerticalBox);
+void FWaveformTransformationTrimFadeRenderer::SetWaveformTransformation(TObjectPtr<UWaveformTransformationBase> InTransformation)
+{
+	StrongTrimFade = TStrongObjectPtr<UWaveformTransformationTrimFade>(CastChecked<UWaveformTransformationTrimFade>(InTransformation.Get()));
+	check(StrongTrimFade->FadeTransformation);
 
-	TSharedRef<SWindow> MenuWindow = SNew(SWindow)
-		.AutoCenter(EAutoCenter::None)
-		.ScreenPosition(LocalCursorPosition)
-		.SizingRule(ESizingRule::Autosized)
-		.SupportsTransparency(EWindowTransparency::None)
-		.IsPopupWindow(true)
-		.CreateTitleBar(false)
-		[
-			MenuContent
-		];
+	FWaveformTransformationFadeRenderer::SetWaveformTransformation(StrongTrimFade->FadeTransformation);
 
-	FadeModeMenuWindow = MenuWindow;
 	TWeakObjectPtr<UWaveformTransformationTrimFade> InWeakTrimFade = StrongTrimFade.Get();
-
-	TSharedPtr<FDelegateHandle> ApplicationActivationStateHandlePtr = MakeShared<FDelegateHandle>();
-	*ApplicationActivationStateHandlePtr = FSlateApplication::Get().OnApplicationActivationStateChanged().AddLambda([InFadeModeMenuWindow = FadeModeMenuWindow, ApplicationActivationStateHandlePtr](bool isActive)
+	OnFadeModeChanged.BindLambda([InWeakTrimFade]()
 		{
-			if (!isActive)
+			if (TStrongObjectPtr<UWaveformTransformationTrimFade> PinnedTrimFade = InWeakTrimFade.Pin())
 			{
-				if (TSharedPtr<SWindow> LockedFadeModeMenuWindow = InFadeModeMenuWindow.Pin())
-				{
-					FSlateApplication::Get().RequestDestroyWindow(LockedFadeModeMenuWindow.ToSharedRef());
-					FSlateApplication::Get().OnApplicationActivationStateChanged().Remove(*ApplicationActivationStateHandlePtr);
-				}
+				PinnedTrimFade->PostEditChange();
 			}
 		});
-
-	if (ApplicationActivationStateHandle.IsValid())
-	{
-		FSlateApplication::Get().OnApplicationActivationStateChanged().Remove(ApplicationActivationStateHandle);
-	}
-
-	ApplicationActivationStateHandle = *ApplicationActivationStateHandlePtr;
-
-	// If focus is lost on the popup, destroy it to prevent popups hanging around
-	TSharedPtr<FDelegateHandle> PopupHandlePtr = MakeShared<FDelegateHandle>();
-	*PopupHandlePtr = FSlateApplication::Get().OnFocusChanging().AddLambda([InFadeModeMenuWindow = FadeModeMenuWindow, PopupHandlePtr](const FFocusEvent& FocusEvent, const FWeakWidgetPath& WeakWidgetPath
-		, const TSharedPtr<SWidget>& OldWidget, const FWidgetPath& WidgetPath, const TSharedPtr<SWidget>& NewWidget)
-		{
-			if (InFadeModeMenuWindow != nullptr && InFadeModeMenuWindow.IsValid())
-			{
-				if (TSharedPtr<SWindow> LockedFadeModeMenuWindow = InFadeModeMenuWindow.Pin())
-				{
-					if (OldWidget && !OldWidget->IsHovered() && LockedFadeModeMenuWindow == OldWidget)
-					{
-						FSlateApplication::Get().RequestDestroyWindow(LockedFadeModeMenuWindow.ToSharedRef());
-						FSlateApplication::Get().OnFocusChanging().Remove(*PopupHandlePtr);
-					}
-				}
-			}
-		});
-
-	if (PopupHandle.IsValid())
-	{
-		FSlateApplication::Get().OnFocusChanging().Remove(PopupHandle);
-	}
-
-	PopupHandle = *PopupHandlePtr;
-
-	for (const TPair<EWaveEditorFadeMode, TSubclassOf<UFadeFunction>>& FadeOptionPair : UWaveformTransformationTrimFade::FadeModeToFadeFunctionMap)
-	{
-		MenuContent->AddSlot()
-			.Padding(5)
-			.AutoHeight()
-			[
-				SNew(SButton)
-					.OnClicked_Lambda([InPopupHandle = PopupHandle, InApplicationActivationStateHandle = ApplicationActivationStateHandle, InWeakFadeModeMenuWindow = FadeModeMenuWindow, InWeakTrimFade, InteractionType = TrimFadeInteractionType, FadeOptionPair, PropertyName = GetPropertyEditedByCurrentInteraction()]() -> FReply
-					{
-						if (TStrongObjectPtr<UWaveformTransformationTrimFade> PinnedTrimFade = InWeakTrimFade.Pin())
-						{
-							if (GEditor && GEditor->Trans)
-							{
-								GEditor->BeginTransaction(TEXT("PropertyEditor"), FText::Format(NSLOCTEXT("PropertyEditor", "SetCurveProperty", "Set Curve"), PropertyName), PinnedTrimFade.Get());
-							}
-
-							PinnedTrimFade->Modify();
-
-							if (InteractionType == FWaveformTransformationTrimFadeRenderer::ETrimFadeInteractionType::RightClickFadeIn)
-							{
-								if (PinnedTrimFade->FadeFunctions.FadeIn)
-								{
-									const float Duration = PinnedTrimFade->FadeFunctions.FadeIn->Duration;
-									PinnedTrimFade->FadeFunctions.FadeIn = NewObject<UFadeFunction>(PinnedTrimFade.Get(), FadeOptionPair.Value.Get(), NAME_None, RF_Transactional);
-									PinnedTrimFade->FadeFunctions.FadeIn->Duration =  Duration;
-								}
-							}
-							else if (InteractionType == FWaveformTransformationTrimFadeRenderer::ETrimFadeInteractionType::RightClickFadeOut)
-							{
-								if (PinnedTrimFade->FadeFunctions.FadeOut)
-								{
-									const float Duration = PinnedTrimFade->FadeFunctions.FadeOut->Duration;
-									PinnedTrimFade->FadeFunctions.FadeOut = NewObject<UFadeFunction>(PinnedTrimFade.Get(), FadeOptionPair.Value.Get(), NAME_None, RF_Transactional);
-									PinnedTrimFade->FadeFunctions.FadeOut->Duration = Duration;
-								}
-							}
-
-							if (GEditor && GEditor->Trans)
-							{
-								GEditor->EndTransaction();
-							}
-
-							PinnedTrimFade->PostEditChange();
-							PinnedTrimFade->OnTransformationChanged.ExecuteIfBound(true);
-						}
-
-						TSharedPtr<SWindow> LockedFadeModeMenuWindow = InWeakFadeModeMenuWindow.Pin();
-						if (LockedFadeModeMenuWindow)
-						{
-							FSlateApplication::Get().RequestDestroyWindow(LockedFadeModeMenuWindow.ToSharedRef());
-
-							if (InPopupHandle.IsValid())
-							{
-								FSlateApplication::Get().OnFocusChanging().Remove(InPopupHandle);
-							}
-
-							if (InApplicationActivationStateHandle.IsValid())
-							{
-								FSlateApplication::Get().OnApplicationActivationStateChanged().Remove(InApplicationActivationStateHandle);
-							}
-						}
-
-						return FReply::Handled();
-					})
-					[
-						SNew(SHorizontalBox)
-						+ SHorizontalBox::Slot()
-						.AutoWidth()
-						.VAlign(VAlign_Center)
-						.HAlign(HAlign_Left)
-						[
-							SNew(SImage)
-							.Image(GetFadeModeIcon(FadeOptionPair.Key).GetIcon())
-						]
-						+ SHorizontalBox::Slot()
-						.VAlign(VAlign_Center)
-						.HAlign(HAlign_Center)
-						.Padding(5)
-						[
-							SNew(STextBlock)
-							.Text(FText::FromString(StaticEnum<EWaveEditorFadeMode>()->GetNameStringByValue(static_cast<int64>(FadeOptionPair.Key))))
-						]
-					]
-			];
-	}
-
-	FSlateApplication::Get().AddWindow(MenuWindow);
-
-	FTSTicker::GetCoreTicker().AddTicker(FTickerDelegate::CreateLambda([InFadeModeMenuWindow = FadeModeMenuWindow, LocalCursorPosition, LocalWindowMaxPosition](float DeltaTime)
-		{
-			if (TSharedPtr<SWindow> LockedFadeModeMenuWindow = InFadeModeMenuWindow.Pin())
-			{
-				const FVector2D WindowSize = LockedFadeModeMenuWindow->GetSizeInScreen();
-
-				if (WindowSize.X > 0.0 && LocalCursorPosition.X + WindowSize.X > LocalWindowMaxPosition.X)
-				{
-					check(WindowSize.Y > 0.0);
-
-					const double XPosition = FMath::Max(0.0, LocalCursorPosition.X - WindowSize.X);
-					LockedFadeModeMenuWindow->MoveWindowTo(FVector2D(XPosition, LocalCursorPosition.Y));
-				}
-			}
-
-			return false; // return false to only execute once
-		}));
 }
 
-const FSlateIcon FWaveformTransformationTrimFadeRenderer::GetFadeModeIcon(const EWaveEditorFadeMode& FadeMode)
+void FWaveformTransformationTrimFadeRenderer::SetPropertyValueDependingOnInteractionType(const FPointerEvent& MouseEvent, const FGeometry& WidgetGeometry)
 {
-	if (TrimFadeInteractionType == FWaveformTransformationTrimFadeRenderer::ETrimFadeInteractionType::RightClickFadeIn)
-	{
-		switch (FadeMode)
-		{
-		case EWaveEditorFadeMode::Linear:
-			return FSlateIcon(FAudioWidgetsStyle::StyleName, "AudioWidgetsStyle.FadeInLinear");
-		case EWaveEditorFadeMode::Exponential:
-			return FSlateIcon(FAudioWidgetsStyle::StyleName, "AudioWidgetsStyle.FadeInExponential");
-		case EWaveEditorFadeMode::Logarithmic:
-			return FSlateIcon(FAudioWidgetsStyle::StyleName, "AudioWidgetsStyle.FadeInLogarithmic");
-		case EWaveEditorFadeMode::Sigmoid:
-			return FSlateIcon(FAudioWidgetsStyle::StyleName, "AudioWidgetsStyle.FadeInSigmoid");
-		default:
-			return FSlateIcon();
-		}
-	}
-	else if (TrimFadeInteractionType == FWaveformTransformationTrimFadeRenderer::ETrimFadeInteractionType::RightClickFadeOut)
-	{
-		switch (FadeMode)
-		{
-		case EWaveEditorFadeMode::Linear:
-			return FSlateIcon(FAudioWidgetsStyle::StyleName, "AudioWidgetsStyle.FadeOutLinear");
-		case EWaveEditorFadeMode::Exponential:
-			return FSlateIcon(FAudioWidgetsStyle::StyleName, "AudioWidgetsStyle.FadeOutExponential");
-		case EWaveEditorFadeMode::Logarithmic:
-			return FSlateIcon(FAudioWidgetsStyle::StyleName, "AudioWidgetsStyle.FadeOutLogarithmic");
-		case EWaveEditorFadeMode::Sigmoid:
-			return FSlateIcon(FAudioWidgetsStyle::StyleName, "AudioWidgetsStyle.FadeOutSigmoid");
-		default:
-			return FSlateIcon();
-		}
-	}
-
-	return FSlateIcon();
-}
+	check(StrongFade);
+	check(StrongTrimFade);
+	check(PixelsPerFrame > 0.0);
+	check(StrongFade->GetSampleRate() > 0);
+	check(WidgetGeometry.GetLocalSize().X > 0.f);
+	const FVector2D LocalCursorPosition = GetLocalCursorPosition(MouseEvent, WidgetGeometry);
+	const float LocalCursorXRatio = FMath::Clamp(LocalCursorPosition.X / WidgetGeometry.GetLocalSize().X, 0.f, 1.f);
+	const double SelectedTime = ConvertXRatioToTime(LocalCursorXRatio);
+	const float NumFrames = TransformationWaveInfo.TotalNumSamples / TransformationWaveInfo.NumChannels;
+	const float MaxDuration = NumFrames / StrongFade->GetSampleRate();
 
-FWaveformTransformationTrimFadeRenderer::ETrimFadeInteractionType FWaveformTransformationTrimFadeRenderer::GetInteractionTypeFromCursorPosition(const FVector2D& InLocalCursorPosition, const FKey MouseButton, const FGeometry& WidgetGeometry) const
-{
-	if (MouseButton == EKeys::LeftMouseButton)
+	if (StrongFade->GetFadeRegions().IsValidIndex(FadeInteractionTypeAndIndex.Value))
 	{
-		if (IsCursorInFadeInInteractionRange(InLocalCursorPosition, WidgetGeometry))
-		{
-			return ETrimFadeInteractionType::ScrubbingFadeIn;
-		}
+		check(StrongFade->GetFadeRegions()[FadeInteractionTypeAndIndex.Value].FadeIn);
+		check(StrongFade->GetFadeRegions()[FadeInteractionTypeAndIndex.Value].FadeOut);
 
-		if (IsCursorInFadeOutInteractionRange(InLocalCursorPosition, WidgetGeometry))
+		switch (FadeInteractionTypeAndIndex.Key)
 		{
-			return ETrimFadeInteractionType::ScrubbingFadeOut;
-		}
-
-		if (StartTimeInteractionXRange.Contains(InLocalCursorPosition.X) && StartTimeInteractionYRange.Contains(InLocalCursorPosition.Y))
+		case FWaveformTransformationFadeRenderer::EFadeInteractionType::None:
+			break;
+		case FWaveformTransformationFadeRenderer::EFadeInteractionType::ScrubbingFadeInStart:
 		{
-			return ETrimFadeInteractionType::ScrubbingLeftHandle;
-		}
+			StrongTrimFade->StartTime = FMath::Min(SelectedTime, StrongTrimFade->EndTime);
 
-		if (EndTimeInteractionXRange.Contains(InLocalCursorPosition.X) && EndTimeInteractionYRange.Contains(InLocalCursorPosition.Y))
-		{
-			return ETrimFadeInteractionType::ScrubbingRightHandle;
+			break;
 		}
-	}
-	else if (MouseButton == EKeys::RightMouseButton)
-	{
-		if (IsCursorInFadeInInteractionRange(InLocalCursorPosition, WidgetGeometry))
+		case FWaveformTransformationFadeRenderer::EFadeInteractionType::ScrubbingFadeInEnd:
 		{
-			return ETrimFadeInteractionType::RightClickFadeIn;
-		}
+			float StartFadeDurationValue = 0;
+			StartFadeDurationValue = FMath::Clamp(SelectedTime - (FadeRegions[FadeInteractionTypeAndIndex.Value].FadeInStartX / PixelsPerFrame) / StrongFade->GetSampleRate(), 0.f, MaxDuration);
+			StrongFade->GetFadeRegions()[FadeInteractionTypeAndIndex.Value].FadeIn->Duration = StartFadeDurationValue;
 
-		if (IsCursorInFadeOutInteractionRange(InLocalCursorPosition, WidgetGeometry))
-		{
-			return ETrimFadeInteractionType::RightClickFadeOut;
+			break;
 		}
-	}
-
-	return ETrimFadeInteractionType::None;
-}
-
-bool FWaveformTransformationTrimFadeRenderer::IsCursorInFadeInInteractionRange(const FVector2D& InLocalCursorPosition, const FGeometry& WidgetGeometry) const
-{
-	return FadeInInteractionXRange.Contains(InLocalCursorPosition.X)
-		&& InLocalCursorPosition.Y < InteractionHandleSize;
-}
-
-bool FWaveformTransformationTrimFadeRenderer::IsCursorInFadeOutInteractionRange(const FVector2D& InLocalCursorPosition, const FGeometry& WidgetGeometry) const
-{
-	return FadeOutInteractionXRange.Contains(InLocalCursorPosition.X)
-		&& InLocalCursorPosition.Y < InteractionHandleSize;
-}
-
-FText FWaveformTransformationTrimFadeRenderer::GetPropertyEditedByCurrentInteraction() const
-{
-	FText OutPropertyName = FText::FromName(NAME_None);
-
-	switch (TrimFadeInteractionType)
-	{
-	case FWaveformTransformationTrimFadeRenderer::ETrimFadeInteractionType::None:
-		break;
-	case FWaveformTransformationTrimFadeRenderer::ETrimFadeInteractionType::ScrubbingLeftHandle:
-		OutPropertyName = FText::FromName(GET_MEMBER_NAME_CHECKED(UWaveformTransformationTrimFade, StartTime));
-		break;
-	case FWaveformTransformationTrimFadeRenderer::ETrimFadeInteractionType::ScrubbingRightHandle:
-		OutPropertyName = FText::FromName(GET_MEMBER_NAME_CHECKED(UWaveformTransformationTrimFade, EndTime));
-		break;
-	case FWaveformTransformationTrimFadeRenderer::ETrimFadeInteractionType::ScrubbingFadeIn:
-		if (StrongTrimFade->FadeFunctions.FadeIn)
+		case FWaveformTransformationFadeRenderer::EFadeInteractionType::ScrubbingFadeOutStart:
 		{
-			OutPropertyName = FText::Format(FText::FromString("{0}::{1}"),
-				FText::FromName(GET_MEMBER_NAME_CHECKED(FFadeFunctionData, FadeIn)),
-				FText::FromName(GET_MEMBER_NAME_CHECKED(UFadeFunction, Duration)));
+			float EndFadeDurationValue = 0;
+			const float diff = ((FadeRegions[FadeInteractionTypeAndIndex.Value].FadeOutEndX / PixelsPerFrame) / StrongFade->GetSampleRate()) - SelectedTime;
+			EndFadeDurationValue = FMath::Clamp(diff, 0.f, MaxDuration);
+			StrongFade->GetFadeRegions()[FadeInteractionTypeAndIndex.Value].FadeOut->Duration = EndFadeDurationValue;
+
+			break;
 		}
-		break;
-	case FWaveformTransformationTrimFadeRenderer::ETrimFadeInteractionType::ScrubbingFadeOut:
-		if (StrongTrimFade->FadeFunctions.FadeOut)
+		case FWaveformTransformationFadeRenderer::EFadeInteractionType::ScrubbingFadeOutEnd:
 		{
-			OutPropertyName = FText::Format(FText::FromString("{0}::{1}"),
-				FText::FromName(GET_MEMBER_NAME_CHECKED(FFadeFunctionData, FadeOut)),
-				FText::FromName(GET_MEMBER_NAME_CHECKED(UFadeFunction, Duration)));
+			StrongTrimFade->EndTime = FMath::Max(SelectedTime, StrongTrimFade->StartTime);
+
+			break;
+		}
+		case FWaveformTransformationFadeRenderer::EFadeInteractionType::RightClickFadeIn:
+		case FWaveformTransformationFadeRenderer::EFadeInteractionType::RightClickFadeOut:
+			ShowSelectFadeModeMenuAtCursor(WidgetGeometry, MouseEvent);
+			break;
+		default:
+			break;
 		}
-		break;
-	default:
-		break;
 	}
-
-	return MoveTemp(OutPropertyName);
-
 }
\ No newline at end of file
diff --git a/Engine/Plugins/Editor/WaveformEditor/Source/WaveformTransformationsWidgets/Private/WaveformTransformationTrimFadeRenderer.h b/Engine/Plugins/Editor/WaveformEditor/Source/WaveformTransformationsWidgets/Private/WaveformTransformationTrimFadeRenderer.h
--- a/Engine/Plugins/Editor/WaveformEditor/Source/WaveformTransformationsWidgets/Private/WaveformTransformationTrimFadeRenderer.h
+++ b/Engine/Plugins/Editor/WaveformEditor/Source/WaveformTransformationsWidgets/Private/WaveformTransformationTrimFadeRenderer.h
@@ -3,82 +3,25 @@
 #pragma once
 
 #include "WaveformTransformationRendererBase.h"
+
+#include "WaveformTransformationFadeRenderer.h"
 #include "WaveformTransformationTrimFade.h"
 
-#include "Brushes/SlateRoundedBoxBrush.h"
 
-class FWaveformTransformationTrimFadeRenderer : public FWaveformTransformationRendererBase
+class FWaveformTransformationTrimFadeRenderer : public FWaveformTransformationFadeRenderer
 {
 public:
-	FWaveformTransformationTrimFadeRenderer() : RoundedBoxBrush(FLinearColor::Green, 2) {}
-	~FWaveformTransformationTrimFadeRenderer();
 
-	virtual int32 OnPaint(const FPaintArgs& Args, const FGeometry& AllottedGeometry, const FSlateRect& MyCullingRect, FSlateWindowElementList& OutDrawElements, int32 LayerId, const FWidgetStyle& InWidgetStyle, bool bParentEnabled) const override;
 	virtual void Tick(const FGeometry& AllottedGeometry, const double InCurrentTime, const float InDeltaTime) override;
 	virtual FReply OnMouseButtonDown(SWidget& OwnerWidget, const FGeometry& MyGeometry, const FPointerEvent& MouseEvent) override;
 	virtual FReply OnMouseButtonUp(SWidget& OwnerWidget, const FGeometry& MyGeometry, const FPointerEvent& MouseEvent) override;
 	virtual FReply OnMouseMove(SWidget& OwnerWidget, const FGeometry& MyGeometry, const FPointerEvent& MouseEvent) override;
-	virtual FReply OnMouseWheel(SWidget& OwnerWidget, const FGeometry& MyGeometry, const FPointerEvent& MouseEvent) override;
-	virtual FCursorReply OnCursorQuery(const FGeometry& MyGeometry, const FPointerEvent& CursorEvent) const override;
 
 	virtual void SetWaveformTransformation(TObjectPtr<UWaveformTransformationBase> InTransformation) override;
 
 private:
-	int32 DrawTrimHandles(const FGeometry& AllottedGeometry, FSlateWindowElementList& OutDrawElements, int32 LayerId) const;
-	int32 DrawFadeCurves(const FGeometry& AllottedGeometry, FSlateWindowElementList& OutDrawElements, int32 LayerId) const;
-	void GenerateFadeCurves(const FGeometry& AllottedGeometry);
-	void UpdateInteractionRange(const FGeometry& AllottedGeometry);
-
-	bool IsCursorInFadeInInteractionRange(const FVector2D& InLocalCursorPosition, const FGeometry& WidgetGeometry) const;
-	bool IsCursorInFadeOutInteractionRange(const FVector2D& InLocalCursorPosition, const FGeometry& WidgetGeometry) const;
-
-	FText GetPropertyEditedByCurrentInteraction() const;
-	FVector2D GetLocalCursorPosition(const FPointerEvent& MouseEvent, const FGeometry& EventGeometry) const;
-	double ConvertXRatioToTime(const float InRatio) const;
-
-	float StartTimeHandleX = 0.f;
-	float EndTimeHandleX = 0.f;
-	uint32 FadeInStartX = 0;
-	uint32 FadeInEndX = 0;
-	uint32 FadeOutStartX = 0;
-	uint32 FadeOutEndX = 0;
-	TArray<FVector2D> FadeInCurvePoints;
-	TArray<FVector2D> FadeOutCurvePoints;
-
-	TRange<float> StartTimeInteractionXRange = TRange<float>::Inclusive(0.f, 0.f);
-	TRange<float> StartTimeInteractionYRange = TRange<float>::Inclusive(0.f, 0.f);
-	TRange<float> EndTimeInteractionXRange = TRange<float>::Inclusive(0.f, 0.f);
-	TRange<float> EndTimeInteractionYRange = TRange<float>::Inclusive(0.f, 0.f);
-	TRange<float> FadeInInteractionXRange = TRange<float>::Inclusive(0.f, 0.f);
-	TRange<float> FadeOutInteractionXRange = TRange<float>::Inclusive(0.f, 0.f);
 
-	double PixelsPerFrame = 0.0;
-	const float InteractionHandleSize = 30.0f;
-	FVector2D MousePosition;
-	int32 TransactionResult = INDEX_NONE;
-
-	enum class ETrimFadeInteractionType : uint8
-	{
-		None = 0,
-		ScrubbingLeftHandle,
-		ScrubbingRightHandle,
-		ScrubbingFadeIn, 
-		ScrubbingFadeOut,
-		RightClickFadeIn,
-		RightClickFadeOut
-	} TrimFadeInteractionType = ETrimFadeInteractionType::None;
-
-	ETrimFadeInteractionType GetInteractionTypeFromCursorPosition(const FVector2D& InLocalCursorPosition, const FKey MouseButton, const FGeometry& WidgetGeometry) const;
 	void SetPropertyValueDependingOnInteractionType(const FPointerEvent& MouseEvent, const FGeometry& WidgetGeometry);
-	void ShowSelectFadeModeMenuAtCursor(const FGeometry& MyGeometry, const FPointerEvent& MouseEvent);
-	const FSlateIcon GetFadeModeIcon(const EWaveEditorFadeMode& FadeMode);
-
-	const FSlateRoundedBoxBrush RoundedBoxBrush;
-
-	TWeakPtr<SWindow> FadeModeMenuWindow = nullptr;
 
 	TStrongObjectPtr<UWaveformTransformationTrimFade> StrongTrimFade = nullptr;
-
-	FDelegateHandle PopupHandle;
-	FDelegateHandle ApplicationActivationStateHandle;
 };
diff --git a/Engine/Plugins/Editor/WaveformEditor/Source/WaveformTransformationsWidgets/Private/WaveformTransformationsWidgetsModule.cpp b/Engine/Plugins/Editor/WaveformEditor/Source/WaveformTransformationsWidgets/Private/WaveformTransformationsWidgetsModule.cpp
--- a/Engine/Plugins/Editor/WaveformEditor/Source/WaveformTransformationsWidgets/Private/WaveformTransformationsWidgetsModule.cpp
+++ b/Engine/Plugins/Editor/WaveformEditor/Source/WaveformTransformationsWidgets/Private/WaveformTransformationsWidgetsModule.cpp
@@ -4,15 +4,18 @@
 
 #include "Modules/ModuleManager.h"
 #include "WaveformTransformationLog.h"
+#include "WaveformTransformationFade.h"
+#include "WaveformTransformationFadeRenderer.h"
+#include "WaveformTransformationMarkers.h"
+#include "WaveformTransformationMarkersRenderer.h"
 #include "WaveformTransformationRendererMapper.h"
 #include "WaveformTransformationTrimFade.h"
 #include "WaveformTransformationTrimFadeRenderer.h"
-#include "WaveformTransformationMarkers.h"
-#include "WaveformTransformationMarkersRenderer.h"
 
 void FWaveformTransformationsWidgetsModule::StartupModule()
 {
 	FWaveformTransformationRendererMapper& RendererMapper = FWaveformTransformationRendererMapper::Get();
+	RendererMapper.RegisterRenderer<FWaveformTransformationFadeRenderer>(UWaveformTransformationFade::StaticClass());
 	RendererMapper.RegisterRenderer<FWaveformTransformationTrimFadeRenderer>(UWaveformTransformationTrimFade::StaticClass());
 	RendererMapper.RegisterRenderer<FWaveformTransformationMarkerRenderer>(UWaveformTransformationMarkers::StaticClass()); 
 }
diff --git a/Engine/Source/Runtime/Engine/Private/SoundWave.cpp b/Engine/Source/Runtime/Engine/Private/SoundWave.cpp
--- a/Engine/Source/Runtime/Engine/Private/SoundWave.cpp
+++ b/Engine/Source/Runtime/Engine/Private/SoundWave.cpp
@@ -4923,14 +4923,18 @@ const FWaveTransformUObjectConfiguration& USoundWave::UpdateTransformations()
 		return TransformationChainConfig;
 	}
 
+	// USoundWave::TotalSamples is named incorrectly and actually represents TotalNumFrames (Duration * SampleRate) 
+	// Can't use Duration * SampleRate here as Duration is updated as the waveform changes length
+	const int32 OriginalNumFrames = FMath::FloorToInt32(TotalSamples);
+
 	TransformationChainConfig.NumChannels = NumChannels;
 	TransformationChainConfig.SampleRate = ImportedSampleRate > 0 ? ImportedSampleRate : GetSampleRateForCurrentPlatform();
 	TransformationChainConfig.StartTime = 0.f;
-	TransformationChainConfig.EndTime = TotalSamples / TransformationChainConfig.SampleRate;
+	TransformationChainConfig.EndTime = OriginalNumFrames / TransformationChainConfig.SampleRate;
 	TransformationChainConfig.WaveCues = CuePoints;
 	TransformationChainConfig.bIsPreviewingLoopRegion = false;
 	TransformationChainConfig.bCachedSoundWaveLoopState = static_cast<bool>(bLooping);
-	TransformationChainConfig.SetOriginalNumFrames(Duration * TransformationChainConfig.SampleRate);
+	TransformationChainConfig.SetOriginalNumFrames(OriginalNumFrames);
 
 	TSet<UClass*> TransformationClassSet;
 	for (int32 i = 0; i < Transformations.Num(); ++i)
