[Sound Wave] - LUFS and Sample Peak Metadata Zeroed out when previewing a Sound Wave in the Content Browser

#Synced-CL 50779174

diff --git a/Engine/Plugins/Editor/WaveformEditor/Source/WaveformEditor/Private/WaveformEditorModule.cpp b/Engine/Plugins/Editor/WaveformEditor/Source/WaveformEditor/Private/WaveformEditorModule.cpp
--- a/Engine/Plugins/Editor/WaveformEditor/Source/WaveformEditor/Private/WaveformEditorModule.cpp
+++ b/Engine/Plugins/Editor/WaveformEditor/Source/WaveformEditor/Private/WaveformEditorModule.cpp
@@ -73,6 +73,20 @@ void FWaveformEditorModule::StartupModule()
 				}
 			}
 
+			OnAssetLoadedHandle = FCoreUObjectDelegates::OnAssetLoaded.AddLambda([](TObjectPtr<UObject> Object)
+				{
+					if (Object && Object->IsA<USoundWave>())
+					{
+						TObjectPtr<USoundWave> SoundWave = CastChecked<USoundWave>(Object);
+
+						// if unanalyzed
+						if (SoundWave->LUFS == 0.f && SoundWave->SamplePeakDB == 0.f)
+						{
+							AnalyzeSoundWaveLoudness(Object);
+						}
+					}
+				});
+
 			OnPackageReloadedHandle = FCoreUObjectDelegates::OnPackageReloaded.AddLambda([](EPackageReloadPhase Phase, FPackageReloadedEvent* ReloadedEvent)
 				{
 					if (Phase == EPackageReloadPhase::PostPackageFixup)
@@ -113,6 +127,7 @@ void FWaveformEditorModule::ShutdownModule()
 		}
 	}
 
+	FCoreUObjectDelegates::OnAssetLoaded.Remove(OnAssetLoadedHandle);
 	FCoreUObjectDelegates::OnPackageReloaded.Remove(OnPackageReloadedHandle);
 }
 
diff --git a/Engine/Plugins/Editor/WaveformEditor/Source/WaveformEditor/Public/WaveformEditorModule.h b/Engine/Plugins/Editor/WaveformEditor/Source/WaveformEditor/Public/WaveformEditorModule.h
--- a/Engine/Plugins/Editor/WaveformEditor/Source/WaveformEditor/Public/WaveformEditorModule.h
+++ b/Engine/Plugins/Editor/WaveformEditor/Source/WaveformEditor/Public/WaveformEditorModule.h
@@ -21,6 +21,7 @@ private:
 
 	FDelegateHandle OnPostEngineInitHandle;
 	FDelegateHandle OnAssetPostImportHandle;
+	FDelegateHandle OnAssetLoadedHandle;
 	FDelegateHandle OnPackageReloadedHandle;
 };
 
