[Waveform Editor] Changing the Fade handle icons and making the colors more colorblind friendly

#Synced-CL 50934197

diff --git a/Engine/Plugins/Editor/WaveformEditor/Source/WaveformTransformationsWidgets/Private/WaveformTransformationFadeRenderer.cpp b/Engine/Plugins/Editor/WaveformEditor/Source/WaveformTransformationsWidgets/Private/WaveformTransformationFadeRenderer.cpp
--- a/Engine/Plugins/Editor/WaveformEditor/Source/WaveformTransformationsWidgets/Private/WaveformTransformationFadeRenderer.cpp
+++ b/Engine/Plugins/Editor/WaveformEditor/Source/WaveformTransformationsWidgets/Private/WaveformTransformationFadeRenderer.cpp
@@ -16,6 +16,14 @@
 #include "Widgets/Text/STextBlock.h"
 #include "Blueprint/WidgetLayoutLibrary.h"
 
+namespace WaveformFadeRendererColors
+{
+	const FLinearColor FadeOutColor = FLinearColor(1.f, 0.1f, 0.f, 1.f);
+	const FLinearColor FadeInColor = FLinearColor::Green;
+	const FLinearColor SelectionColor = FLinearColor::Yellow;
+	const FLinearColor LineColor = FLinearColor::Gray;
+}
+
 FWaveformTransformationFadeRenderer::~FWaveformTransformationFadeRenderer()
 {
 	if (FSlateApplication::IsInitialized())
@@ -49,7 +57,7 @@ int32 FWaveformTransformationFadeRenderer::DrawOffsetHandles(const FGeometry& Al
 			HandleGeometry,
 			&RoundedBoxBrush,
 			ESlateDrawEffect::None,
-			FLinearColor::Green
+			WaveformFadeRendererColors::FadeInColor
 		);
 
 		HandleGeometry = AllottedGeometry.ToPaintGeometry(FVector2f(HandleWidth, HandleWidth), FSlateLayoutTransform(FVector2f(FadeRegion.FadeOutEndX - HandleWidth, AllottedGeometry.Size.Y - HandleWidth)));
@@ -60,7 +68,7 @@ int32 FWaveformTransformationFadeRenderer::DrawOffsetHandles(const FGeometry& Al
 			HandleGeometry,
 			&RoundedBoxBrush,
 			ESlateDrawEffect::None,
-			FLinearColor::Red
+			WaveformFadeRendererColors::FadeOutColor
 		);
 	}
 
@@ -69,13 +77,14 @@ int32 FWaveformTransformationFadeRenderer::DrawOffsetHandles(const FGeometry& Al
 
 int32 FWaveformTransformationFadeRenderer::DrawFadeCurves(const FGeometry& AllottedGeometry, FSlateWindowElementList& OutDrawElements, int32 LayerId) const
 {
-	const FLinearColor FadeColor = FLinearColor::Gray;
 	const float HandleWidth = InteractionHandleSize;
 	const float HandleHeight = InteractionHandleSize;
-	const float HandleOffset = 5.0f;
+	const float HandleOffset = 0.0f;
 
-	const FLinearColor HandleColor = FLinearColor::White;
-	const FLinearColor HandleSelectedColor = FLinearColor::Yellow;
+	const FSlateBrush* FadeInHandleBrush = FSlateIcon(FAudioWidgetsStyle::StyleName, "AudioWidgetsStyle.FadeStart").GetIcon();
+	const FSlateBrush* FadeOutHandleBrush = FSlateIcon(FAudioWidgetsStyle::StyleName, "AudioWidgetsStyle.FadeEnd").GetIcon();
+
+	FPaintGeometry HandleGeometry;
 
 	for (const FFadeRegion& FadeRegion : FadeRegions)
 	{
@@ -90,37 +99,27 @@ int32 FWaveformTransformationFadeRenderer::DrawFadeCurves(const FGeometry& Allot
 				AllottedGeometry.ToPaintGeometry(),
 				FadeRegion.FadeInCurvePoints,
 				ESlateDrawEffect::None,
-				FadeColor
+				WaveformFadeRendererColors::LineColor
 			);
 
-			const float	HandleStart = FadeRegion.FadeInCurvePoints.Last(0).X;
-			const FPaintGeometry HandleGeometry = AllottedGeometry.ToPaintGeometry(FVector2f(HandleWidth, HandleHeight), FSlateLayoutTransform(FVector2f(HandleStart - HandleOffset, 0)));
-			const FSlateBrush* HandleBrush = FAppStyle::GetBrush("WorldBrowser.DirectionXPositive");
-
-			FSlateDrawElement::MakeBox(
-				OutDrawElements,
-				++LayerId,
-				HandleGeometry,
-				HandleBrush,
-				ESlateDrawEffect::DisabledEffect,
-				IsFadeInSelected ? HandleSelectedColor : HandleColor
-			);
+			const float HandleStart = FadeRegion.FadeInCurvePoints.Last(0).X;
+			HandleGeometry = AllottedGeometry.ToPaintGeometry(FVector2f(HandleWidth, HandleHeight), FSlateLayoutTransform(FVector2f(HandleStart - HandleOffset, 0)));
 		}
 		else
 		{
-			FPaintGeometry HandleGeometry = AllottedGeometry.ToPaintGeometry(FVector2f(HandleWidth, HandleHeight), FSlateLayoutTransform(FVector2f(FadeRegion.FadeInStartX- HandleOffset, 0)));
-			const FSlateBrush* HandleBrush = FAppStyle::GetBrush("WorldBrowser.DirectionXPositive");
-
-			FSlateDrawElement::MakeBox(
-				OutDrawElements,
-				++LayerId,
-				HandleGeometry,
-				HandleBrush,
-				ESlateDrawEffect::DisabledEffect,
-				IsFadeInSelected ? HandleSelectedColor : HandleColor
-			);
+			HandleGeometry = AllottedGeometry.ToPaintGeometry(FVector2f(HandleWidth, HandleHeight), FSlateLayoutTransform(FVector2f(FadeRegion.FadeInStartX- HandleOffset, 0)));
 		}
 
+		check(FadeInHandleBrush);
+		FSlateDrawElement::MakeBox(
+			OutDrawElements,
+			++LayerId,
+			HandleGeometry,
+			FadeInHandleBrush,
+			ESlateDrawEffect::None,
+			IsFadeInSelected ? WaveformFadeRendererColors::SelectionColor : WaveformFadeRendererColors::FadeInColor
+		);
+
 		if (FadeRegion.FadeOutCurvePoints.Num() > 0)
 		{
 			FSlateDrawElement::MakeLines(
@@ -129,36 +128,26 @@ int32 FWaveformTransformationFadeRenderer::DrawFadeCurves(const FGeometry& Allot
 				AllottedGeometry.ToPaintGeometry(),
 				FadeRegion.FadeOutCurvePoints,
 				ESlateDrawEffect::None,
-				FadeColor
+				WaveformFadeRendererColors::LineColor
 			);
 
-			const float	HandleStart = FadeRegion.FadeOutCurvePoints[0].X - HandleWidth;
-			const FPaintGeometry HandleGeometry = AllottedGeometry.ToPaintGeometry(FVector2f(HandleWidth, HandleHeight), FSlateLayoutTransform(FVector2f(HandleStart + HandleOffset, 0)));
-			const FSlateBrush* HandleBrush = FAppStyle::GetBrush("WorldBrowser.DirectionXNegative");
-
-			FSlateDrawElement::MakeBox(
-				OutDrawElements,
-				++LayerId,
-				HandleGeometry,
-				HandleBrush,
-				ESlateDrawEffect::DisabledEffect,
-				IsFadeOutSelected ? HandleSelectedColor : HandleColor
-			);
+			const float HandleStart = FadeRegion.FadeOutCurvePoints[0].X - HandleWidth;
+			HandleGeometry = AllottedGeometry.ToPaintGeometry(FVector2f(HandleWidth, HandleHeight), FSlateLayoutTransform(FVector2f(HandleStart + HandleOffset, 0)));
 		}
 		else
 		{
-			FPaintGeometry HandleGeometry = AllottedGeometry.ToPaintGeometry(FVector2f(HandleWidth, HandleHeight), FSlateLayoutTransform(FVector2f(FadeRegion.FadeOutEndX - HandleWidth + HandleOffset, 0)));
-			const FSlateBrush* HandleBrush = FAppStyle::GetBrush("WorldBrowser.DirectionXNegative");
-
-			FSlateDrawElement::MakeBox(
-				OutDrawElements,
-				++LayerId,
-				HandleGeometry,
-				HandleBrush,
-				ESlateDrawEffect::DisabledEffect,
-				IsFadeOutSelected ? HandleSelectedColor : HandleColor
-			);
+			HandleGeometry = AllottedGeometry.ToPaintGeometry(FVector2f(HandleWidth, HandleHeight), FSlateLayoutTransform(FVector2f(FadeRegion.FadeOutEndX - HandleWidth + HandleOffset, 0)));
 		}
+
+		check(FadeOutHandleBrush);
+		FSlateDrawElement::MakeBox(
+			OutDrawElements,
+			++LayerId,
+			HandleGeometry,
+			FadeOutHandleBrush,
+			ESlateDrawEffect::None,
+			IsFadeOutSelected ? WaveformFadeRendererColors::SelectionColor : WaveformFadeRendererColors::FadeOutColor
+		);
 	}
 
 	return LayerId;
diff --git a/Engine/Plugins/Runtime/AudioWidgets/Source/AudioWidgets/Private/AudioWidgetsStyle.cpp b/Engine/Plugins/Runtime/AudioWidgets/Source/AudioWidgets/Private/AudioWidgetsStyle.cpp
--- a/Engine/Plugins/Runtime/AudioWidgets/Source/AudioWidgets/Private/AudioWidgetsStyle.cpp
+++ b/Engine/Plugins/Runtime/AudioWidgets/Source/AudioWidgets/Private/AudioWidgetsStyle.cpp
@@ -194,6 +194,7 @@ void FAudioWidgetsStyle::SetupStyles()
 {
 	const FVector2f Icon16x16 = FVector2f(16.0f, 16.0f);
 	const FVector2f Icon20x20 = FVector2f(20.0f, 20.0f);
+	const FVector2f Icon30x30 = FVector2f(30.0f, 30.0f);
 
 	Set("AudioWidgetsStyle.FadeIn", new IMAGE_BRUSH_SVG(TEXT("Slate/Icons/FadeIn_20"), Icon20x20));
 	Set("AudioWidgetsStyle.FadeInLinear", new IMAGE_BRUSH_SVG(TEXT("Slate/Icons/FadeInLinear_16"), Icon16x16));
@@ -207,8 +208,16 @@ void FAudioWidgetsStyle::SetupStyles()
 	Set("AudioWidgetsStyle.FadeOutExponential", new IMAGE_BRUSH_SVG(TEXT("Slate/Icons/FadeOutExponential_16"), Icon16x16));
 	Set("AudioWidgetsStyle.FadeOutSigmoid", new IMAGE_BRUSH_SVG(TEXT("Slate/Icons/FadeOutSigmoid_16"), Icon16x16));
 
+	Set("AudioWidgetsStyle.FadeStart", new IMAGE_BRUSH_SVG(TEXT("Slate/Icons/FadeStart_30"), Icon30x30));
+	Set("AudioWidgetsStyle.FadeEnd", new IMAGE_BRUSH_SVG(TEXT("Slate/Icons/FadeEnd_30"), Icon30x30));
+
 	Set("AudioWidgetsStyle.MarkerAdd", new IMAGE_BRUSH_SVG(TEXT("Slate/Icons/Marker_20"), Icon20x20));
 	Set("AudioWidgetsStyle.MarkerDelete", new IMAGE_BRUSH_SVG(TEXT("Slate/Icons/MarkerDelete_20"), Icon20x20));
+
+	Set("AudioWidgetsStyle.LoopStart", new IMAGE_BRUSH_SVG(TEXT("Slate/Icons/LoopStart_30"), Icon30x30));
+	Set("AudioWidgetsStyle.LoopEnd", new IMAGE_BRUSH_SVG(TEXT("Slate/Icons/LoopEnd_30"), Icon30x30));
+
+	Set("AudioWidgetsStyle.Reset", new CORE_IMAGE_BRUSH_SVG("Starship/Common/Reset", Icon16x16));
 }
 
 FAudioWidgetsStyle::~FAudioWidgetsStyle()
