[Sound Wave] - LUFS and Sample Peak Metadata Zeroed out when previewing a legacy Sound Wave in the Content Browser that was added before the loudness measurements happened on import

#was-backout 50794418

#Synced-CL 50934197

diff --git a/Engine/Plugins/Editor/WaveformEditor/Source/WaveformEditor/Private/WaveformEditor.cpp b/Engine/Plugins/Editor/WaveformEditor/Source/WaveformEditor/Private/WaveformEditor.cpp
--- a/Engine/Plugins/Editor/WaveformEditor/Source/WaveformEditor/Private/WaveformEditor.cpp
+++ b/Engine/Plugins/Editor/WaveformEditor/Source/WaveformEditor/Private/WaveformEditor.cpp
@@ -147,7 +147,10 @@ bool FWaveformEditor::Init(const EToolkitMode::Type Mode, const TSharedPtr<ITool
 		BindDelegatesForActiveTransformations();
 
 #if WITH_EDITOR
-		UpdateSoundWaveProperties();
+		if (SoundWave->LUFS == 0.f && SoundWave->SamplePeakDB == 0.f)
+		{
+			UpdateSoundWaveProperties();
+		}
 #endif // WITH_EDITOR
 	}
 
@@ -1878,6 +1881,9 @@ void FWaveformEditor::UpdateSoundWaveProperties()
 	const float LUFS = WaveformAudioAnalysis::GetLUFS(AudioBuffer, SampleRate, TransformationChainConfig.NumChannels);
 
 	SoundWave->SetLoudnessValues(LUFS, SamplePeakDB);
+
+	// Update SoundWave tooltip values to show new loudness values in the content browser when hovering over the sound
+	SoundWave->PostEditChange();
 }
 #endif //WITH_EDITOR
 
diff --git a/Engine/Plugins/Editor/WaveformEditor/Source/WaveformEditor/Private/WaveformEditorModule.cpp b/Engine/Plugins/Editor/WaveformEditor/Source/WaveformEditor/Private/WaveformEditorModule.cpp
--- a/Engine/Plugins/Editor/WaveformEditor/Source/WaveformEditor/Private/WaveformEditorModule.cpp
+++ b/Engine/Plugins/Editor/WaveformEditor/Source/WaveformEditor/Private/WaveformEditorModule.cpp
@@ -7,6 +7,7 @@
 #include "Modules/ModuleManager.h"
 #include "PropertyEditorModule.h"
 #include "Sound/SoundWave.h"
+#include "Sound/SoundWaveProcedural.h"
 #include "TransformedWaveformViewFactory.h"
 #include "WaveformAudioAnalysisFunctions.h"
 #include "WaveformEditorCommands.h"
@@ -21,7 +22,7 @@ namespace
 {
 	void AnalyzeSoundWaveLoudness(UObject* Object)
 	{
-		if (Object && Object->IsA<USoundWave>())
+		if (Object && Object->IsA<USoundWave>() && !Object->IsA<USoundWaveProcedural>())
 		{
 			TObjectPtr<USoundWave> SoundWave = CastChecked<USoundWave>(Object);
 
@@ -71,6 +72,19 @@ void FWaveformEditorModule::StartupModule()
 							AnalyzeSoundWaveLoudness(Object);
 						});
 				}
+
+				OnSoundBasePreviewHandle = GEditor->OnSoundBasePreview.AddLambda([](TObjectPtr<USoundBase> SoundBase)
+					{
+						if (SoundBase && SoundBase->IsA<USoundWave>())
+						{
+							TObjectPtr<USoundWave> SoundWave = CastChecked<USoundWave>(SoundBase);
+
+							if (SoundWave->LUFS == 0.f && SoundWave->SamplePeakDB == 0.f)
+							{
+								AnalyzeSoundWaveLoudness(SoundBase);
+							}
+						}
+					});
 			}
 
 			OnPackageReloadedHandle = FCoreUObjectDelegates::OnPackageReloaded.AddLambda([](EPackageReloadPhase Phase, FPackageReloadedEvent* ReloadedEvent)
@@ -111,6 +125,8 @@ void FWaveformEditorModule::ShutdownModule()
 		{
 			ImportSubsystem->OnAssetPostImport.Remove(OnAssetPostImportHandle);
 		}
+
+		GEditor->OnSoundBasePreview.Remove(OnSoundBasePreviewHandle);
 	}
 
 	FCoreUObjectDelegates::OnPackageReloaded.Remove(OnPackageReloadedHandle);
diff --git a/Engine/Plugins/Editor/WaveformEditor/Source/WaveformEditor/Public/WaveformEditorModule.h b/Engine/Plugins/Editor/WaveformEditor/Source/WaveformEditor/Public/WaveformEditorModule.h
--- a/Engine/Plugins/Editor/WaveformEditor/Source/WaveformEditor/Public/WaveformEditorModule.h
+++ b/Engine/Plugins/Editor/WaveformEditor/Source/WaveformEditor/Public/WaveformEditorModule.h
@@ -21,6 +21,7 @@ private:
 
 	FDelegateHandle OnPostEngineInitHandle;
 	FDelegateHandle OnAssetPostImportHandle;
+	FDelegateHandle OnSoundBasePreviewHandle;
 	FDelegateHandle OnPackageReloadedHandle;
 };
 
diff --git a/Engine/Source/Editor/UnrealEd/Classes/Editor/EditorEngine.h b/Engine/Source/Editor/UnrealEd/Classes/Editor/EditorEngine.h
--- a/Engine/Source/Editor/UnrealEd/Classes/Editor/EditorEngine.h
+++ b/Engine/Source/Editor/UnrealEd/Classes/Editor/EditorEngine.h
@@ -666,6 +666,9 @@ public:
 	DECLARE_MULTICAST_DELEGATE_TwoParams(FPostBugItGoCalled, const FVector& Loc, const FRotator& Rot);
 	FPostBugItGoCalled PostBugItGoCalled;
 
+	DECLARE_MULTICAST_DELEGATE_OneParam(FOnSoundBasePreview, TObjectPtr<USoundBase>)
+	FOnSoundBasePreview OnSoundBasePreview;
+
 	/** Whether or not the editor is currently compiling */
 	bool bIsCompiling;
 
diff --git a/Engine/Source/Editor/UnrealEd/Private/EditorEngine.cpp b/Engine/Source/Editor/UnrealEd/Private/EditorEngine.cpp
--- a/Engine/Source/Editor/UnrealEd/Private/EditorEngine.cpp
+++ b/Engine/Source/Editor/UnrealEd/Private/EditorEngine.cpp
@@ -2955,6 +2955,8 @@ UAudioComponent* UEditorEngine::PlayPreviewSound(USoundBase* Sound,  USoundNode*
 	if(UAudioComponent* AudioComponent = ResetPreviewAudioComponent(Sound, SoundNode))
 	{
 		AudioComponent->Play();
+		OnSoundBasePreview.Broadcast(Sound);
+
 		return AudioComponent;
 	}
 
