[Waveform Editor] - UE.EditorAutomation(RunTest=Audio) Win64 - Waveform Editor Export failed Original soundwave was modified

#Synced-CL 50969872

diff --git a/Engine/Plugins/Editor/WaveformEditor/Source/WaveformEditor/Private/WaveformEditor.cpp b/Engine/Plugins/Editor/WaveformEditor/Source/WaveformEditor/Private/WaveformEditor.cpp
--- a/Engine/Plugins/Editor/WaveformEditor/Source/WaveformEditor/Private/WaveformEditor.cpp
+++ b/Engine/Plugins/Editor/WaveformEditor/Source/WaveformEditor/Private/WaveformEditor.cpp
@@ -149,7 +149,15 @@ bool FWaveformEditor::Init(const EToolkitMode::Type Mode, const TSharedPtr<ITool
 #if WITH_EDITOR
 		if (SoundWave->LUFS == 0.f && SoundWave->SamplePeakDB == 0.f)
 		{
+			// Update the properties without marking the SoundWave dirty to prevent checkouts unintended by the user
+			const bool bWasAlreadyDirty = SoundWave->GetPackage()->IsDirty();
+
 			UpdateSoundWaveProperties();
+
+			if (!bWasAlreadyDirty)
+			{
+				SoundWave->GetPackage()->ClearDirtyFlag();
+			}
 		}
 #endif // WITH_EDITOR
 	}
