[Waveform Editor] Warning in Cook Content Worker fix: missing script on import for fades. Removed fade content migration out of trimfade constructor and now it is only in PostLoad

#Synced-CL 50969872

diff --git a/Engine/Plugins/Editor/WaveformEditor/Source/WaveformTransformations/Private/WaveformTransformationTrimFade.cpp b/Engine/Plugins/Editor/WaveformEditor/Source/WaveformTransformations/Private/WaveformTransformationTrimFade.cpp
--- a/Engine/Plugins/Editor/WaveformEditor/Source/WaveformTransformations/Private/WaveformTransformationTrimFade.cpp
+++ b/Engine/Plugins/Editor/WaveformEditor/Source/WaveformTransformations/Private/WaveformTransformationTrimFade.cpp
@@ -268,81 +268,6 @@ void FWaveTransformationTrimFade::ProcessAudio(Audio::FWaveformTransformationWav
 UWaveformTransformationTrimFade::UWaveformTransformationTrimFade(const FObjectInitializer& ObjectInitializer) : Super(ObjectInitializer)
 {
 	FadeTransformation = CreateDefaultSubobject<UWaveformTransformationFade>(TEXT("FadeTransformation"));
-
-	PRAGMA_DISABLE_DEPRECATION_WARNINGS
-	if (FadeFunctions.FadeIn == nullptr)
-	{
-		if (StartFadeTime != 0.f)
-		{
-			if (StartFadeCurve > WaveformTrasnformationTrimFadePrivate::Linear)
-			{
-				FadeFunctions.FadeIn = CreateDefaultSubobject<UFadeCurveFunctionExponential>(TEXT("FadeInFunction"));
-			}
-			else if (StartFadeCurve < 0.f)
-			{
-				FadeFunctions.FadeIn = CreateDefaultSubobject<UFadeCurveFunctionSigmoid>(TEXT("FadeInFunction"));
-			}
-			else if (StartFadeCurve < WaveformTrasnformationTrimFadePrivate::Linear)
-			{
-				FadeFunctions.FadeIn = CreateDefaultSubobject<UFadeCurveFunctionLogarithmic>(TEXT("FadeInFunction"));
-			}
-			else if (StartFadeCurve == WaveformTrasnformationTrimFadePrivate::Linear)
-			{
-				FadeFunctions.FadeIn = CreateDefaultSubobject<UFadeFunctionLinear>(TEXT("FadeInFunction"));
-			}
-
-			FadeFunctions.FadeIn->Duration = StartFadeTime;
-			WaveformTrasnformationTrimFadePrivate::SetFadeCurve(FadeFunctions.FadeIn, StartFadeCurve, StartSCurveSharpness);
-
-			StartFadeTime = 0.f;
-		}
-		else
-		{
-			FadeFunctions.FadeIn = CreateDefaultSubobject<UFadeCurveFunctionExponential>(TEXT("FadeInFunction"));
-			FadeFunctions.FadeIn->Duration = 0.f;
-		}
-	}
-
-	if (FadeFunctions.FadeOut == nullptr)
-	{
-		if (EndFadeTime != 0.f)
-		{
-			if (EndFadeCurve > WaveformTrasnformationTrimFadePrivate::Linear)
-			{
-				FadeFunctions.FadeOut = CreateDefaultSubobject<UFadeCurveFunctionExponential>(TEXT("FadeOutFunction"));
-			}
-			else if (EndFadeCurve < 0.f)
-			{
-				FadeFunctions.FadeOut = CreateDefaultSubobject<UFadeCurveFunctionSigmoid>(TEXT("FadeOutFunction"));
-			}
-			else if (EndFadeCurve < WaveformTrasnformationTrimFadePrivate::Linear)
-			{
-				FadeFunctions.FadeOut = CreateDefaultSubobject<UFadeCurveFunctionLogarithmic>(TEXT("FadeOutFunction"));
-			}
-			else if (EndFadeCurve == WaveformTrasnformationTrimFadePrivate::Linear)
-			{
-				FadeFunctions.FadeOut = CreateDefaultSubobject<UFadeFunctionLinear>(TEXT("FadeOutFunction"));
-			}
-
-			FadeFunctions.FadeOut->Duration = EndFadeTime;
-			WaveformTrasnformationTrimFadePrivate::SetFadeCurve(FadeFunctions.FadeOut, EndFadeCurve, EndSCurveSharpness);
-
-			EndFadeTime = 0.f;
-		}
-		else
-		{
-			FadeFunctions.FadeOut = CreateDefaultSubobject<UFadeCurveFunctionExponential>(TEXT("FadeOutFunction"));
-			FadeFunctions.FadeOut->Duration = 0.f;
-		}
-	}
-
-	if (FadeFunctions.FadeIn || FadeFunctions.FadeOut)
-	{
-		check(FadeTransformation);
-
-		FadeFunctions.MigrateToTransformationFadeFunction(FadeTransformation, GetMutableFadeRegion());
-	}
-	PRAGMA_ENABLE_DEPRECATION_WARNINGS
 }
 
 #if WITH_EDITOR
@@ -484,6 +409,9 @@ void UWaveformTransformationTrimFade::PostLoad()
 		check(FadeTransformation);
 
 		FadeFunctions.MigrateToTransformationFadeFunction(FadeTransformation, GetMutableFadeRegion());
+
+		FadeFunctions.FadeIn = nullptr;
+		FadeFunctions.FadeOut = nullptr;
 	}
 	PRAGMA_ENABLE_DEPRECATION_WARNINGS
 
