[Waveform Editor] - Remove deprecated fade code

#Synced-CL 51092343

diff --git a/Engine/Plugins/Editor/WaveformEditor/Source/WaveformTransformations/Private/WaveformTransformationTrimFade.cpp b/Engine/Plugins/Editor/WaveformEditor/Source/WaveformTransformations/Private/WaveformTransformationTrimFade.cpp
--- a/Engine/Plugins/Editor/WaveformEditor/Source/WaveformTransformations/Private/WaveformTransformationTrimFade.cpp
+++ b/Engine/Plugins/Editor/WaveformEditor/Source/WaveformTransformations/Private/WaveformTransformationTrimFade.cpp
@@ -99,92 +99,12 @@ namespace WaveformTrasnformationTrimFadePrivate
 	PRAGMA_ENABLE_DEPRECATION_WARNINGS
 }
 
-const float UWaveformTransformationTrimFade::MinFadeCurve = -0.1f;
-const float UWaveformTransformationTrimFade::MaxFadeCurve = 10.f;
-
 namespace WaveformTransformationTrimFadeNames
 {
 	static FLazyName StartTimeName(GET_MEMBER_NAME_CHECKED(UWaveformTransformationTrimFade, StartTime));
 	static FLazyName EndTimeName(GET_MEMBER_NAME_CHECKED(UWaveformTransformationTrimFade, EndTime));
 }
 
-PRAGMA_DISABLE_DEPRECATION_WARNINGS
-static void ApplyFadeIn(Audio::FAlignedFloatBuffer& InputAudio, const FFadeFunctionData& FadeFunctions, const int32 NumChannels, const float SampleRate)
-{
-	if (FadeFunctions.FadeIn == nullptr)
-	{
-		return;
-	}
-
-	check(NumChannels > 0);
-	check(SampleRate > 0.f);
-
-	const float FadeDuration = FadeFunctions.FadeIn->Duration;
-	const int32 FadeNumFrames = FMath::Min(FadeDuration * SampleRate, InputAudio.Num() / NumChannels);
-
-	if (InputAudio.Num() < NumChannels || FadeDuration < SMALL_NUMBER || FadeNumFrames == 0)
-	{
-		return;
-	}
-
-	float* InputPtr = InputAudio.GetData();
-
-	for (int32 FrameIndex = 0; FrameIndex < FadeNumFrames; ++FrameIndex)
-	{
-		const float FadeFraction = static_cast<float>(FrameIndex) / static_cast<float>(FadeNumFrames);
-		check(FadeFraction <= 1.0f);
-		check(FadeFraction >= 0.0f);
-
-		const float EnvValue = FadeFunctions.FadeIn->GetFadeInCurveValue(FadeFraction);
-
-		for (int32 ChannelIt = 0; ChannelIt < NumChannels; ++ChannelIt)
-		{
-			*(InputPtr + ChannelIt) *= EnvValue;
-		}
-
-		InputPtr += NumChannels;
-	}
-}
-
-static void ApplyFadeOut(Audio::FAlignedFloatBuffer& InputAudio, const FFadeFunctionData& FadeFunctions, const int32 NumChannels, const float SampleRate)
-{
-	if (FadeFunctions.FadeOut == nullptr)
-	{
-		return;
-	}
-
-	check(NumChannels > 0);
-	check(SampleRate > 0.f);
-
-	const float FadeDuration = FadeFunctions.FadeOut->Duration;
-	const int32 FadeNumFrames = FMath::Min(FadeDuration * SampleRate, InputAudio.Num() / NumChannels);
-
-	if (InputAudio.Num() < NumChannels || FadeDuration < SMALL_NUMBER || FadeNumFrames == 0)
-	{
-		return;
-	}
-
-	const int32 StartSampleIndex = InputAudio.Num() - (FadeNumFrames * NumChannels);
-	float* InputPtr = &InputAudio[StartSampleIndex];
-
-	for (int32 FrameIndex = 0; FrameIndex < FadeNumFrames; ++FrameIndex)
-	{
-		const float FadeFraction = static_cast<float>(FrameIndex) / static_cast<float>(FadeNumFrames);
-		check(FadeFraction <= 1.0f);
-		check(FadeFraction >= 0.0f);
-
-		const float EnvValue = FadeFunctions.FadeOut->GetFadeOutCurveValue(FadeFraction);
-
-		for (int32 ChannelIt = 0; ChannelIt < NumChannels; ++ChannelIt)
-		{
-			*(InputPtr + ChannelIt) *= EnvValue;
-		}
-
-		InputPtr += NumChannels;
-	}
-}
-PRAGMA_ENABLE_DEPRECATION_WARNINGS
-
 FWaveTransformationTrimFade::FWaveTransformationTrimFade(TObjectPtr<UWaveformTransformationFade> InFadeTransformation, double InStartTime, double InEndTime)
 	: FadeTransformation(InFadeTransformation)
 	, StartTime(InStartTime)
@@ -454,64 +374,6 @@ FTransformationFadeFunctionData& UWaveformTransformationTrimFade::GetMutableFade
 	return FadeTransformation->GetMutableFadeRegions()[0];
 }
 
-PRAGMA_DISABLE_DEPRECATION_WARNINGS
-const double UWaveformTransformationTrimFade::GetFadeInCurveValue(const float StartFadeCurve, const double FadeFraction, const float SCurveSharpness)
-{
-	double CurveValue = 1;
-
-	if (StartFadeCurve < 0)
-	{
-		const double X = FadeFraction;
-		const double A = static_cast<float>(-FMath::Clamp(SCurveSharpness, 0.1f, 1.0f));
-		const double K = 20 * FMath::Pow(A, 5.0);
-
-		if (FadeFraction <= 0.5)
-		{
-			CurveValue = -(FMath::Tanh(K * PI * FMath::LogX(10, WaveformTrasnformationTrimFadePrivate::XMultiplier * X) + ((K * PI) / 2))) / 2 + 0.5;
-		}
-		else
-		{
-			CurveValue = (FMath::Tanh(K * PI * FMath::LogX(10, WaveformTrasnformationTrimFadePrivate::XMultiplier * (-X + 1)) + ((K * PI) / 2))) / 2 + 0.5;
-		}
-	}
-	else
-	{
-		CurveValue = FMath::Pow(FadeFraction, StartFadeCurve);
-	}
-
-	return CurveValue;
-}
-
-const double UWaveformTransformationTrimFade::GetFadeOutCurveValue(const float EndFadeCurve, const double FadeFraction, const float SCurveSharpness)
-{
-	double CurveValue = 1;
-
-	if (EndFadeCurve < 0)
-	{
-		const double X = FadeFraction;
-		const double A = static_cast<float>(FMath::Clamp(SCurveSharpness, 0.1f, 1.0f));
-		const double K = 20 * FMath::Pow(A, 5.0);
-
-		if (FadeFraction <= 0.5)
-		{
-			CurveValue = -(FMath::Tanh(K * PI * FMath::LogX(10, WaveformTrasnformationTrimFadePrivate::XMultiplier * X) + ((K * PI) / 2))) / 2 + 0.5;
-		}
-		else
-		{
-			CurveValue = (FMath::Tanh(K * PI * FMath::LogX(10, WaveformTrasnformationTrimFadePrivate::XMultiplier * (-X + 1)) + ((K * PI) / 2))) / 2 + 0.5;
-		}
-	}
-	else
-	{
-		// Changed from 1.f - FMath::Pow(FadeFraction, FadeCurve) so that the fade out curve is a horizontally mirrored version of fade in 
-		// instead of vertically mirrored
-		CurveValue = FMath::Pow(-FadeFraction + 1, EndFadeCurve);
-	}
-
-	return CurveValue;
-}
-PRAGMA_ENABLE_DEPRECATION_WARNINGS
-
 void UWaveformTransformationTrimFade::UpdateDurationProperties(const float InAvailableDuration)
 {
 	check(InAvailableDuration > 0.f);
diff --git a/Engine/Plugins/Editor/WaveformEditor/Source/WaveformTransformations/Public/WaveformTransformationTrimFade.h b/Engine/Plugins/Editor/WaveformEditor/Source/WaveformTransformations/Public/WaveformTransformationTrimFade.h
--- a/Engine/Plugins/Editor/WaveformEditor/Source/WaveformTransformations/Public/WaveformTransformationTrimFade.h
+++ b/Engine/Plugins/Editor/WaveformEditor/Source/WaveformTransformations/Public/WaveformTransformationTrimFade.h
@@ -120,15 +120,6 @@ struct UE_DEPRECATED(5.8, "UFadeFunction has been deprecated. Use UTransformatio
 
 	void MigrateToTransformationFadeFunction(const TObjectPtr<UWaveformTransformationFade> InFadeTransformation, FTransformationFadeFunctionData& NewFadeRegion);
 };
-
-UENUM()
-enum class UE_DEPRECATED(5.8, "UFadeFunction has been deprecated. Use UTransformationFadeFunction.") EWaveEditorFadeMode : uint8
-{
-	Linear = 0,
-	Exponential,
-	Logarithmic,
-	Sigmoid
-};
 PRAGMA_ENABLE_DEPRECATION_WARNINGS
 
 class FWaveTransformationTrimFade : public Audio::IWaveTransformation
@@ -222,19 +213,6 @@ public:
 	UPROPERTY(EditAnywhere, Category = "Fade", meta = (ClampMin = 0, ClampMax = 1.0, DisplayName = "S-Curve Sharpness", EditCondition = "EndFadeTime != 0 && EndFadeCurve < 0", EditConditionHides))
 	float EndSCurveSharpness = 0.75;
 
-	UE_DEPRECATED(5.7, "Individual fade properties have been deprecated. Use UTransformationFadeFunction instead.")
-	static UE_API const float MinFadeCurve;
-
-	UE_DEPRECATED(5.7, "Individual fade properties have been deprecated. Use UTransformationFadeFunction instead.")
-	static UE_API const float MaxFadeCurve;
-	
-	PRAGMA_DISABLE_DEPRECATION_WARNINGS
-	UE_DEPRECATED(5.7, "Individual fade properties have been deprecated. Use UTransformationFadeFunction instead.")
-	static UE_API const TMap<EWaveEditorFadeMode, float> FadeModeToCurveValueMap;
-
-	static UE_API const TMap<EWaveEditorFadeMode, TSubclassOf<UFadeFunction>> FadeModeToFadeFunctionMap;
-	PRAGMA_ENABLE_DEPRECATION_WARNINGS
-
 	UE_API virtual Audio::FTransformationPtr CreateTransformation() const override;
 
 	UE_API void UpdateConfiguration(FWaveTransformUObjectConfiguration& InOutConfiguration) override;
@@ -242,12 +220,6 @@ public:
 
 	UE_API const FTransformationFadeFunctionData& GetFadeRegion() const;
 	UE_API FTransformationFadeFunctionData& GetMutableFadeRegion() const;
-
-	UE_DEPRECATED(5.7, "Individual fade properties have been deprecated. Use UTransformationFadeFunction instead.")
-	static UE_API const double GetFadeInCurveValue(const float StartFadeCurve, const double FadeFraction, const float SCurveSharpness = 0);
-
-	UE_DEPRECATED(5.7, "Individual fade properties have been deprecated. Use UTransformationFadeFunction instead.")
-	static UE_API const double GetFadeOutCurveValue(const float EndFadeCurve, const double FadeFraction, const float SCurveSharpness = 0);
 	
 private:	
 
