[Waveform Editor] UpdateTransformations doesn't consistently call UpdateSoundWaveProperties()

#Synced-CL 51261445

diff --git a/Engine/Plugins/Editor/WaveformEditor/Source/WaveformEditor/Private/WaveformEditor.cpp b/Engine/Plugins/Editor/WaveformEditor/Source/WaveformEditor/Private/WaveformEditor.cpp
--- a/Engine/Plugins/Editor/WaveformEditor/Source/WaveformEditor/Private/WaveformEditor.cpp
+++ b/Engine/Plugins/Editor/WaveformEditor/Source/WaveformEditor/Private/WaveformEditor.cpp
@@ -795,17 +795,18 @@ void FWaveformEditor::UpdateTransformations(bool bMarkDirty)
 
 	NotifyPostChange(PropertyChangedEvent, &PropertyChain);
 
-	// Check if we do not want to mark the package dirty for a user facing change they would not expect to modify data such as selecting a loop region
-	if (!bMarkDirty && !bWasAlreadyDirty)
-	{
-		SoundWave->GetPackage()->ClearDirtyFlag();
-	}
-	else 
+	// Check if we do want to mark the package dirty or not. The user would not expect to modify data when doing things like selecting a loop region
+	if (bMarkDirty)
 	{
 #if WITH_EDITOR
+		// Only update SoundWave properties when a change that requires marking it dirty occurs
 		UpdateSoundWaveProperties();
 #endif //WITH_EDITOR
 	}
+	else if (!bWasAlreadyDirty)
+	{
+		SoundWave->GetPackage()->ClearDirtyFlag();
+	}
 
 	if (bIsPlaying)
 	{
