[Waveform Editor] Engine crashes when 'fade curve' is set to 0 while 'fade in' is set to logarithmic

#Synced-CL 51529245

diff --git a/Engine/Plugins/Editor/WaveformEditor/Source/WaveformTransformationsWidgets/Private/WaveformTransformationFadeRenderer.cpp b/Engine/Plugins/Editor/WaveformEditor/Source/WaveformTransformationsWidgets/Private/WaveformTransformationFadeRenderer.cpp
--- a/Engine/Plugins/Editor/WaveformEditor/Source/WaveformTransformationsWidgets/Private/WaveformTransformationFadeRenderer.cpp
+++ b/Engine/Plugins/Editor/WaveformEditor/Source/WaveformTransformationsWidgets/Private/WaveformTransformationFadeRenderer.cpp
@@ -188,11 +188,15 @@ void FWaveformTransformationFadeRenderer::GenerateFadeCurves(const FGeometry& Al
 
 			if (DisplayedFadeInPixelLength > 0)
 			{
-				for (uint32 Pixel = 0; Pixel < DisplayedFadeInPixelLength; ++Pixel)
+				// Ensure the first pixel of the line is placed vertically at the bottom to correctly show where the fade starts
+				FadeRegions[Index].FadeInCurvePoints[0] = FVector2D(FadeRegions[Index].FadeInStartX, AllottedGeometry.Size.Y);
+
+				for (uint32 Pixel = 1; Pixel < DisplayedFadeInPixelLength; ++Pixel)
 				{
 					const double FadeFraction = FMath::Clamp(static_cast<float>(Pixel) / FadeInPixelLength, 0.f, 1.f);
 					const double CurveFunction = StrongFade->GetFadeRegions()[Index].FadeIn->GetFadeInCurveValue(FadeFraction);
 
+					// Prevents line flicker when dragging the fade handle
 					const bool bNeedsFlickerPrevention = (Pixel == DisplayedFadeInPixelLength - 1) && DisplayedFadeInPixelLength != 1;
 
 					// Prevent fade line from flickering if the math doesn't work out to the last Y value being 0 as the position is changing
@@ -232,11 +236,15 @@ void FWaveformTransformationFadeRenderer::GenerateFadeCurves(const FGeometry& Al
 
 			if (DisplayedFadeOutPixelLength > 0)
 			{
-				for (uint32 Pixel = 0; Pixel < DisplayedFadeOutPixelLength; ++Pixel)
+				// Ensure the first pixel of the line is placed vertically at the top to correctly show where the fade starts
+				FadeRegions[Index].FadeOutCurvePoints[0] = FVector2D(FadeRegions[Index].FadeOutStartX, 0);
+
+				for (uint32 Pixel = 1; Pixel < DisplayedFadeOutPixelLength; ++Pixel)
 				{
 					const double FadeFraction = FMath::Clamp(static_cast<float>(Pixel) / FadeOutPixelLength, 0.f, 1.f);
 					const double CurveFunction = StrongFade->GetFadeRegions()[Index].FadeOut->GetFadeOutCurveValue(FadeFraction);
 
+					// Prevents line flicker when dragging the fade handle and prevents problems with the logarithmic fade when using a 0 exponent value.
 					const bool bNeedsFlickerPrevention = (Pixel == DisplayedFadeOutPixelLength - 1) && DisplayedFadeOutPixelLength != 1;
 
 					// Prevent fade line from flickering if the math doesn't work out to the last Y value being equal to AllottedGeometry.Size.Y as the 
