Waveform Editor: Fix playhead snapping when editing trims to snap to appropriate trim start or end bounds

#Synced-CL 51601966

diff --git a/Engine/Plugins/Editor/WaveformEditor/Source/WaveformEditor/Private/WaveformEditor.cpp b/Engine/Plugins/Editor/WaveformEditor/Source/WaveformEditor/Private/WaveformEditor.cpp
--- a/Engine/Plugins/Editor/WaveformEditor/Source/WaveformEditor/Private/WaveformEditor.cpp
+++ b/Engine/Plugins/Editor/WaveformEditor/Source/WaveformEditor/Private/WaveformEditor.cpp
@@ -1834,6 +1834,7 @@ void FWaveformEditor::UpdateTransportState()
 		{
 			PlaybackTimeBeforeTransformInteraction = 0;
 			StartTimeBeforeTransformInteraction = TransformationChainConfig.StartTime;
+			UE_LOG(LogTemp, Warning, TEXT("StartTime Reset"));
 		}
 		const float StartTimeDifference = StartTimeBeforeTransformInteraction - TransformationChainConfig.StartTime;
 		const float AdjustedPlaybackTime = PlaybackTimeBeforeTransformInteraction + StartTimeDifference;
@@ -1848,11 +1849,34 @@ void FWaveformEditor::UpdateTransportState()
 			TransformInteractionPlayState = EAudioComponentPlayState::Stopped;
 		case EAudioComponentPlayState::Stopped:
 			// Set TransportController and displayed TransportCoordinator to value clamped within TransformationChainConfig.
-			TransportController->CacheStartTime(AdjustedPlaybackTime);
+
 			if (!TransportCoordinator->IsScrubbing())
 			{
-				const float CachedAudioStartTimeAsPercentage = TransportController->GetCachedAudioStartTimeAsPercentage();
-				TransportCoordinator->SetProgressRatio(CachedAudioStartTimeAsPercentage);
+				check(SoundWave->GetImportedSampleRate() > 0);
+				const float PlayheadTime = static_cast<float>(WaveformView.ViewWidget->GetPlayheadFramePosition()) / static_cast<float>(SoundWave->GetImportedSampleRate());
+
+				// If the playhead goes out of bounds, snap it to the closest trim edge
+				if (PlayheadTime > TransformationChainConfig.EndTime)
+				{
+					TransportController->CacheStartTime(TransformationChainConfig.EndTime - TransformationChainConfig.StartTime);
+					TransportCoordinator->SetProgressRatio(1.0f);
+
+					// Cache a play time at the end if the end time trim is what caused the playhead to move
+					StartTimeBeforeTransformInteraction = TransformationChainConfig.EndTime;
+				}
+				else if (PlayheadTime < TransformationChainConfig.StartTime)
+				{
+					const float CachedAudioStartTimeAsPercentage = TransportController->GetCachedAudioStartTimeAsPercentage();
+					TransportCoordinator->SetProgressRatio(CachedAudioStartTimeAsPercentage);
+				}
+				else
+				{
+					TransportController->CacheStartTime(AdjustedPlaybackTime);
+				}
+			}
+			else
+			{
+				TransportController->CacheStartTime(AdjustedPlaybackTime);
 			}
 			break;
 		default:
diff --git a/Engine/Plugins/Editor/WaveformEditor/Source/WaveformEditorWidgets/Private/STransformedWaveformViewPanel.cpp b/Engine/Plugins/Editor/WaveformEditor/Source/WaveformEditorWidgets/Private/STransformedWaveformViewPanel.cpp
--- a/Engine/Plugins/Editor/WaveformEditor/Source/WaveformEditorWidgets/Private/STransformedWaveformViewPanel.cpp
+++ b/Engine/Plugins/Editor/WaveformEditor/Source/WaveformEditorWidgets/Private/STransformedWaveformViewPanel.cpp
@@ -233,6 +233,11 @@ FReply STransformedWaveformViewPanel::LaunchTimeRulerContextMenu()
 	return FReply::Unhandled();
 }
 
+const int64 STransformedWaveformViewPanel::GetPlayheadFramePosition()
+{
+	return (CachedPlayheadRatio * GridData->GetDisplayRange().Size<uint32>()) + GridData->GetDisplayRange().GetLowerBoundValue();
+}
+
 void STransformedWaveformViewPanel::UpdateDisplayUnit(const ESampledSequenceDisplayUnit InDisplayUnit)
 {
 	DisplayUnit = InDisplayUnit;
diff --git a/Engine/Plugins/Editor/WaveformEditor/Source/WaveformEditorWidgets/Public/STransformedWaveformViewPanel.h b/Engine/Plugins/Editor/WaveformEditor/Source/WaveformEditorWidgets/Public/STransformedWaveformViewPanel.h
--- a/Engine/Plugins/Editor/WaveformEditor/Source/WaveformEditorWidgets/Public/STransformedWaveformViewPanel.h
+++ b/Engine/Plugins/Editor/WaveformEditor/Source/WaveformEditorWidgets/Public/STransformedWaveformViewPanel.h
@@ -62,6 +62,8 @@ public:
 
 	UE_API FReply LaunchTimeRulerContextMenu();
 
+	UE_API const int64 GetPlayheadFramePosition();
+
 private:
 	UE_API void CreateLayout();
 
diff --git a/Engine/Plugins/Editor/WaveformEditor/Source/WaveformEditorWidgets/Public/WaveformEditorGridData.h b/Engine/Plugins/Editor/WaveformEditor/Source/WaveformEditorWidgets/Public/WaveformEditorGridData.h
--- a/Engine/Plugins/Editor/WaveformEditor/Source/WaveformEditorWidgets/Public/WaveformEditorGridData.h
+++ b/Engine/Plugins/Editor/WaveformEditor/Source/WaveformEditorWidgets/Public/WaveformEditorGridData.h
@@ -25,6 +25,8 @@ public:
 	UE_API void SetTicksTimeFont(const FSlateFontInfo* InNewFont);
 	UE_API const float SnapPositionToClosestFrame(const float InPixelPosition) const;
 
+	const TRange<uint32> GetDisplayRange() { return DisplayRange; }
+
 	FOnWaveformEditorGridUpdated OnGridMetricsUpdated;
 
 private:
