Waveform Editor: Playhead refactor to use frame positions to allow better precision, along with handling offscreen playhead positions. Playhead is not rendered when offscreen.

#Synced-CL 51797321

diff --git a/Engine/Plugins/Editor/WaveformEditor/Source/WaveformEditor/Private/WaveformEditor.cpp b/Engine/Plugins/Editor/WaveformEditor/Source/WaveformEditor/Private/WaveformEditor.cpp
--- a/Engine/Plugins/Editor/WaveformEditor/Source/WaveformEditor/Private/WaveformEditor.cpp
+++ b/Engine/Plugins/Editor/WaveformEditor/Source/WaveformEditor/Private/WaveformEditor.cpp
@@ -621,7 +621,7 @@ void FWaveformEditor::CreateMarker(bool bIsLoopRegion)
 
 	FSoundWaveCuePoint NewCuePoint;
 	NewCuePoint.CuePointID = HighestCueID + 1;
-	NewCuePoint.FramePosition = static_cast<int64>(TransportCoordinator->GetFocusPoint() * SoundWave->TotalSamples);
+	NewCuePoint.FramePosition = WaveformView.ViewWidget->GetPlayheadFramePosition();
 
 	if (bIsLoopRegion)
 	{
@@ -690,7 +690,7 @@ void FWaveformEditor::SkipToNextMarker()
 	{
 		check(TransportCoordinator != nullptr);
 		check(SoundWave->TotalSamples > 0);
-		const int64 PlaybackFramePosition = static_cast<int64>(TransportCoordinator->GetFocusPoint() * SoundWave->TotalSamples);
+		const int64 PlaybackFramePosition = WaveformView.ViewWidget->GetPlayheadFramePosition();
 		int64 NewPlaybackFramePosition = MAX_int64;
 
 		// Find the nearest marker beyond the current playhead.
@@ -716,12 +716,8 @@ void FWaveformEditor::SkipToNextMarker()
 			}
 			if (ActiveDuration > 0)
 			{
-				check(TransformationChainConfig.SampleRate > 0);
-				const float ActiveDurationInFrames = ActiveDuration * TransformationChainConfig.SampleRate;
-				const float AdjustedNewPlaybackFramePosition = FMath::Max(NewPlaybackFramePosition - TransformationChainConfig.StartTime * TransformationChainConfig.SampleRate, 0.0f);
-				
-				HandlePlayheadScrub(AdjustedNewPlaybackFramePosition / ActiveDurationInFrames, false);
-				TransportCoordinator->SetProgressRatio(AdjustedNewPlaybackFramePosition / ActiveDurationInFrames);
+				WaveformView.ViewWidget->SetPlayheadFramePosition(NewPlaybackFramePosition);
+				WaveformView.ViewWidget->ResolveFromFramePosition();
 			}
 		}
 	}
@@ -732,7 +728,21 @@ void FWaveformEditor::SkipToNextZeroCrossing()
 	check(SoundWave != nullptr);
 
 	TArray<float> PCMArray = WaveformView.DataProvider->GetTransformedPCMData();
-	int64 StartFrame = FMath::Min(FMath::RoundToInt64(TransportCoordinator->GetPlayBackRatioFromFocusPoint(TransportCoordinator->GetFocusPoint()) * SoundWave->GetNumFrames()), static_cast<int64>(SoundWave->GetNumFrames()-1));
+	const int64 StartFrame = WaveformView.ViewWidget->GetPlayheadFramePosition();
+
+	float EndTime = TransformationChainConfig.EndTime;
+	if (EndTime < 0)
+	{
+		EndTime = SoundWave->Duration;
+	}
+
+	check(TransformationChainConfig.SampleRate > 0);
+	const int64 EndFrame = FMath::Max(EndTime * TransformationChainConfig.SampleRate - 1, 0);
+
+	if (EndFrame <= StartFrame || StartFrame < 0)
+	{
+		return;
+	}
 
 	const UWaveformEditorTransformationsSettings* TransformationsSettings = GetWaveformEditorTransformationsSettings();
 	check(TransformationsSettings);
@@ -742,24 +752,10 @@ void FWaveformEditor::SkipToNextZeroCrossing()
 	check(PCMArray.Num() > 0);
 	check(SoundWave->NumChannels > 0);
 
-	const int64 FramesToMove = UWaveformTransformationMarkers::GetFramesToNextZeroCrossing(StartFrame, SoundWave->GetNumFrames()-1, PCMArray, true, NoiseThresholdValue, SoundWave->NumChannels);
-
-	float ActiveDuration = TransformationChainConfig.EndTime - TransformationChainConfig.StartTime;
-	if (ActiveDuration <= 0.0f)
-	{
-		ActiveDuration = SoundWave->Duration;
-	}
-
-	check(ActiveDuration);
-	check(TransformationChainConfig.SampleRate > 0);
-
-	const float ActiveDurationInFrames = ActiveDuration * TransformationChainConfig.SampleRate;
-	const float AdjustedNewPlaybackFramePosition = FMath::Max((StartFrame + FramesToMove) - 
-		(TransformationChainConfig.StartTime * TransformationChainConfig.SampleRate), 0.0f);
-	const float ZoomedRatio = TransportCoordinator->ConvertAbsoluteRatioToZoomed(AdjustedNewPlaybackFramePosition / ActiveDurationInFrames);
+	const int64 FramesToMove = UWaveformTransformationMarkers::GetFramesToNextZeroCrossing(StartFrame, EndFrame, PCMArray, true, NoiseThresholdValue, SoundWave->NumChannels);
 
-	HandlePlayheadScrub(ZoomedRatio, false);
-	TransportCoordinator->ScrubFocusPoint(ZoomedRatio, false);
+	WaveformView.ViewWidget->SetPlayheadFramePosition(StartFrame + FramesToMove);
+	WaveformView.ViewWidget->ResolveFromFramePosition();
 }
 
 void FWaveformEditor::RegenerateTransformations()
@@ -1053,8 +1049,8 @@ void FWaveformEditor::ExecuteOverwriteTransformations()
 
 void FWaveformEditor::ResetPlaybackToStart()
 {
-	HandlePlayheadScrub(0, false);
-	TransportCoordinator->SetProgressRatio(0);
+	WaveformView.ViewWidget->SetPlayheadFramePosition(TransformationChainConfig.GetStartFrameOffset());
+	WaveformView.ViewWidget->ResolveFromFramePosition();
 }
 
 void FWaveformEditor::RegisterTabSpawners(const TSharedRef<FTabManager>& InTabManager)
@@ -1815,6 +1811,7 @@ void FWaveformEditor::NotifyPostChange(const FPropertyChangedEvent& PropertyChan
 	BindDelegatesForActiveTransformations();
 
 	UpdateTransportState();
+
 	bIsInteractingWithTransformations = false;
 
 	PromptToChangeCuePointOrigin();
@@ -2171,12 +2168,14 @@ void FWaveformEditor::HandleRenderDataUpdate()
 	{
 		TransportCoordinator->UpdatePlaybackRange(WaveformView.DataProvider->GetTransformedWaveformBounds());
 		WaveformView.DataProvider->RequestSequenceView(TransportCoordinator->GetDisplayRange());
+		WaveformView.ViewWidget->ResolveFromFramePosition();
 	}
 }
 
 void FWaveformEditor::HandleDisplayRangeUpdate(const TRange<double> NewRange)
 {
 	WaveformView.DataProvider->RequestSequenceView(NewRange);
+	WaveformView.ViewWidget->ResolveFromFramePosition();
 }
 
 void FWaveformEditor::HandlePlayheadScrub(const float InTargetPlayBackRatio, const bool bIsMoving)
diff --git a/Engine/Plugins/Editor/WaveformEditor/Source/WaveformEditorWidgets/Private/STransformedWaveformViewPanel.cpp b/Engine/Plugins/Editor/WaveformEditor/Source/WaveformEditorWidgets/Private/STransformedWaveformViewPanel.cpp
--- a/Engine/Plugins/Editor/WaveformEditor/Source/WaveformEditorWidgets/Private/STransformedWaveformViewPanel.cpp
+++ b/Engine/Plugins/Editor/WaveformEditor/Source/WaveformEditorWidgets/Private/STransformedWaveformViewPanel.cpp
@@ -188,9 +188,23 @@ void STransformedWaveformViewPanel::ReceiveSequenceView(const FFixedSampledSeque
 	}
 }
 
-void STransformedWaveformViewPanel::SetPlayheadRatio(const float InRatio)
+void STransformedWaveformViewPanel::SetPlayheadRatio(const double InRatio)
 {
+	check(InRatio >= 0.0 && InRatio <= 1.0);
 	CachedPlayheadRatio = InRatio;
+	check(GridData);
+	CachedPlayheadFramePosition = GridData->GetFramePositionFromRatio(InRatio);
+}
+
+void STransformedWaveformViewPanel::SetPlayheadFramePosition(const int64 FramePosition)
+{
+	CachedPlayheadFramePosition = FramePosition;
+}
+
+void STransformedWaveformViewPanel::ResolveFromFramePosition()
+{
+	check(GridData);
+	CachedPlayheadRatio = GridData->GetRatioFromFramePosition(CachedPlayheadFramePosition);
 }
 
 void STransformedWaveformViewPanel::SetOnPlayheadOverlayMouseButtonUp(FPointerEventHandler InEventHandler)
@@ -233,11 +247,6 @@ FReply STransformedWaveformViewPanel::LaunchTimeRulerContextMenu()
 	return FReply::Unhandled();
 }
 
-const int64 STransformedWaveformViewPanel::GetPlayheadFramePosition()
-{
-	return (CachedPlayheadRatio * GridData->GetDisplayRange().Size<uint32>()) + GridData->GetDisplayRange().GetLowerBoundValue();
-}
-
 void STransformedWaveformViewPanel::UpdateDisplayUnit(const ESampledSequenceDisplayUnit InDisplayUnit)
 {
 	DisplayUnit = InDisplayUnit;
@@ -258,6 +267,16 @@ void STransformedWaveformViewPanel::Tick(const FGeometry& AllottedGeometry, cons
 		}
 	}
 
+	if (PlayheadOverlay && GridData)
+	{
+		PlayheadOverlay->SetPlayheadVisible(GridData->IsFramePositionInDisplayRange(CachedPlayheadFramePosition));
+	}
+
+	if (TimeRuler && GridData)
+	{
+		TimeRuler->SetPlayheadVisible(GridData->IsFramePositionInDisplayRange(CachedPlayheadFramePosition));
+	}
+
 	UpdatePlayheadPosition(PaintedWidth);
 
 	if (ZoomController)
@@ -268,11 +287,15 @@ void STransformedWaveformViewPanel::Tick(const FGeometry& AllottedGeometry, cons
 
 void STransformedWaveformViewPanel::UpdatePlayheadPosition(const float PaintedWidth)
 {
+	// If the frame position is not visible then the overlays will not render and do not need to update
+	if (GridData != nullptr && !GridData->IsFramePositionInDisplayRange(CachedPlayheadFramePosition))
+	{
+		return;
+	}
 
 	float PlayheadX = CachedPlayheadRatio * PaintedWidth;
 	PlayheadX = GridData ? GridData->SnapPositionToClosestFrame(PlayheadX) : PlayheadX;
 
-
 	if (PlayheadOverlay)
 	{
 		PlayheadOverlay->SetPlayheadPosition(PlayheadX);
diff --git a/Engine/Plugins/Editor/WaveformEditor/Source/WaveformEditorWidgets/Private/WaveformEditorGridData.cpp b/Engine/Plugins/Editor/WaveformEditor/Source/WaveformEditorWidgets/Private/WaveformEditorGridData.cpp
--- a/Engine/Plugins/Editor/WaveformEditor/Source/WaveformEditorWidgets/Private/WaveformEditorGridData.cpp
+++ b/Engine/Plugins/Editor/WaveformEditor/Source/WaveformEditorWidgets/Private/WaveformEditorGridData.cpp
@@ -94,4 +94,34 @@ const float FWaveformEditorGridData::SnapPositionToClosestFrame(const float InPi
 
 
 	return SnappedPosition;
-}
\ No newline at end of file
+}
+
+const int64 FWaveformEditorGridData::GetFramePositionFromRatio(double Ratio) const
+{
+	const int64 DisplayedFrames = DisplayRange.GetUpperBoundValue() - DisplayRange.GetLowerBoundValue();
+
+	const int64 FramePosition = FMath::RoundToInt64(Ratio * DisplayedFrames) + DisplayRange.GetLowerBoundValue();
+
+	return FramePosition;
+}
+
+const double FWaveformEditorGridData::GetRatioFromFramePosition(int64 FramePosition) const
+{
+	if (!IsFramePositionInDisplayRange(FramePosition))
+	{
+		return -1; // Invalid ratio as we are now offscreen
+	}
+
+	const int64 PositionInDisplay = FramePosition - DisplayRange.GetLowerBoundValue();
+	const double Range = DisplayRange.Size<double>();
+
+	// A display range of 0 should not be possible as you would not be able to see anything in waveform editor
+	check(Range > 0);
+
+	return PositionInDisplay / Range;
+}
+
+const bool FWaveformEditorGridData::IsFramePositionInDisplayRange(int64 FramePosition) const
+{
+	return FramePosition >= DisplayRange.GetLowerBoundValue() && FramePosition <= DisplayRange.GetUpperBoundValue();
+}
diff --git a/Engine/Plugins/Editor/WaveformEditor/Source/WaveformEditorWidgets/Public/STransformedWaveformViewPanel.h b/Engine/Plugins/Editor/WaveformEditor/Source/WaveformEditorWidgets/Public/STransformedWaveformViewPanel.h
--- a/Engine/Plugins/Editor/WaveformEditor/Source/WaveformEditorWidgets/Public/STransformedWaveformViewPanel.h
+++ b/Engine/Plugins/Editor/WaveformEditor/Source/WaveformEditorWidgets/Public/STransformedWaveformViewPanel.h
@@ -52,7 +52,9 @@ public:
 	UE_API void Construct(const FArguments& InArgs, const FFixedSampledSequenceView& InView);
 	UE_API virtual void Tick(const FGeometry& AllottedGeometry, const double InCurrentTime, const float InDeltaTime) override;
 	UE_API virtual void ReceiveSequenceView(const FFixedSampledSequenceView InView, const uint32 FirstSampleIndex = 0) override;
-	UE_API void SetPlayheadRatio(const float InRatio);
+	UE_API void SetPlayheadRatio(const double InRatio);
+	UE_API void SetPlayheadFramePosition(const int64 FramePosition);
+	UE_API void ResolveFromFramePosition();
 
 	UE_API void SetOnPlayheadOverlayMouseButtonUp(FPointerEventHandler InEventHandler);
 	UE_API void SetOnTimeRulerMouseButtonUp(FPointerEventHandler InEventHandler);
@@ -62,7 +64,7 @@ public:
 
 	UE_API FReply LaunchTimeRulerContextMenu();
 
-	UE_API const int64 GetPlayheadFramePosition();
+	const int64 GetPlayheadFramePosition() { return CachedPlayheadFramePosition; }
 
 private:
 	UE_API void CreateLayout();
@@ -102,7 +104,8 @@ private:
 
 	FFixedSampledSequenceView DataView;
 	
-	float CachedPlayheadRatio = 0.f;
+	double CachedPlayheadRatio = 0.0;
+	int64 CachedPlayheadFramePosition = 0;
 	UE_API const UWaveformEditorWidgetsSettings* GetWaveformEditorWidgetsSettings();
 };
 
diff --git a/Engine/Plugins/Editor/WaveformEditor/Source/WaveformEditorWidgets/Public/WaveformEditorGridData.h b/Engine/Plugins/Editor/WaveformEditor/Source/WaveformEditorWidgets/Public/WaveformEditorGridData.h
--- a/Engine/Plugins/Editor/WaveformEditor/Source/WaveformEditorWidgets/Public/WaveformEditorGridData.h
+++ b/Engine/Plugins/Editor/WaveformEditor/Source/WaveformEditorWidgets/Public/WaveformEditorGridData.h
@@ -24,6 +24,9 @@ public:
 	UE_API virtual const FFixedSampledSequenceGridMetrics GetGridMetrics() const override;
 	UE_API void SetTicksTimeFont(const FSlateFontInfo* InNewFont);
 	UE_API const float SnapPositionToClosestFrame(const float InPixelPosition) const;
+	UE_API const int64 GetFramePositionFromRatio(double Ratio) const;
+	UE_API const double GetRatioFromFramePosition(int64 FramePosition) const;
+	UE_API const bool IsFramePositionInDisplayRange(int64 FramePosition) const;
 
 	const TRange<uint32> GetDisplayRange() { return DisplayRange; }
 
diff --git a/Engine/Plugins/Runtime/AudioWidgets/Source/AudioWidgets/Private/SFixedSampledSequenceRuler.cpp b/Engine/Plugins/Runtime/AudioWidgets/Source/AudioWidgets/Private/SFixedSampledSequenceRuler.cpp
--- a/Engine/Plugins/Runtime/AudioWidgets/Source/AudioWidgets/Private/SFixedSampledSequenceRuler.cpp
+++ b/Engine/Plugins/Runtime/AudioWidgets/Source/AudioWidgets/Private/SFixedSampledSequenceRuler.cpp
@@ -241,6 +241,11 @@ FReply SFixedSampledSequenceRuler::LaunchContextMenu()
 	return FReply::Handled();
 }
 
+void SFixedSampledSequenceRuler::SetPlayheadVisible(bool Enable)
+{
+	bDisplayPlayhead = Enable;
+}
+
 void SFixedSampledSequenceRuler::DrawPlayheadHandle(const FGeometry& AllottedGeometry, FSlateWindowElementList& OutDrawElements, int32& LayerId) const
 {
 	const float WindowWidth = AllottedGeometry.Size.X;
diff --git a/Engine/Plugins/Runtime/AudioWidgets/Source/AudioWidgets/Private/SPlayheadOverlay.cpp b/Engine/Plugins/Runtime/AudioWidgets/Source/AudioWidgets/Private/SPlayheadOverlay.cpp
--- a/Engine/Plugins/Runtime/AudioWidgets/Source/AudioWidgets/Private/SPlayheadOverlay.cpp
+++ b/Engine/Plugins/Runtime/AudioWidgets/Source/AudioWidgets/Private/SPlayheadOverlay.cpp
@@ -9,6 +9,7 @@ void SPlayheadOverlay::Construct(const FArguments& InArgs)
 	PlayheadWidth = InArgs._Style->PlayheadWidth;
 	DesiredWidth = InArgs._Style->DesiredWidth;
 	DesiredHeight = InArgs._Style->DesiredHeight;
+	bIsVisible = true;
 }
 
 void SPlayheadOverlay::SetPlayheadPosition(const float InNewPosition)
@@ -24,9 +25,17 @@ void SPlayheadOverlay::OnStyleUpdated(const FPlayheadOverlayStyle UpdatedStyle)
 	DesiredHeight = UpdatedStyle.DesiredHeight;
 }
 
+void SPlayheadOverlay::SetPlayheadVisible(bool Enable)
+{
+	bIsVisible = Enable;
+}
+
 int32 SPlayheadOverlay::OnPaint(const FPaintArgs& Args, const FGeometry& AllottedGeometry, const FSlateRect& MyCullingRect, FSlateWindowElementList& OutDrawElements, int32 LayerId, const FWidgetStyle& InWidgetStyle, bool bParentEnabled) const
 {
-	LayerId = DrawPlayhead(AllottedGeometry, OutDrawElements, LayerId);
+	if (bIsVisible)
+	{
+		LayerId = DrawPlayhead(AllottedGeometry, OutDrawElements, LayerId);
+	}
 
 	return LayerId;
 }
diff --git a/Engine/Plugins/Runtime/AudioWidgets/Source/AudioWidgets/Private/SparseSampledSequenceTransportCoordinator.cpp b/Engine/Plugins/Runtime/AudioWidgets/Source/AudioWidgets/Private/SparseSampledSequenceTransportCoordinator.cpp
--- a/Engine/Plugins/Runtime/AudioWidgets/Source/AudioWidgets/Private/SparseSampledSequenceTransportCoordinator.cpp
+++ b/Engine/Plugins/Runtime/AudioWidgets/Source/AudioWidgets/Private/SparseSampledSequenceTransportCoordinator.cpp
@@ -93,7 +93,6 @@ void FSparseSampledSequenceTransportCoordinator::UpdateZoomRatioAndDisplayRange(
 
 	ZoomRatio = ClampedZoomRatio;
 	UpdateDisplayRange(MinDisplayRangeValue, MaxDisplayRangeValue);
-	MoveFocusPoint(NewPlayheadPosition);
 }
 
 void FSparseSampledSequenceTransportCoordinator::Stop()
@@ -125,7 +124,7 @@ void FSparseSampledSequenceTransportCoordinator::MoveWaveformView(float MoveAmou
 	UpdateDisplayRange(MinDisplayRangeValue, MaxDisplayRangeValue);
 }
 
-const float FSparseSampledSequenceTransportCoordinator::GetFocusPoint() const
+const double FSparseSampledSequenceTransportCoordinator::GetFocusPoint() const
 {
 	return FocusPoint;
 }
diff --git a/Engine/Plugins/Runtime/AudioWidgets/Source/AudioWidgets/Public/ISparseSampledSequenceTransportCoordinator.h b/Engine/Plugins/Runtime/AudioWidgets/Source/AudioWidgets/Public/ISparseSampledSequenceTransportCoordinator.h
--- a/Engine/Plugins/Runtime/AudioWidgets/Source/AudioWidgets/Public/ISparseSampledSequenceTransportCoordinator.h
+++ b/Engine/Plugins/Runtime/AudioWidgets/Source/AudioWidgets/Public/ISparseSampledSequenceTransportCoordinator.h
@@ -10,7 +10,7 @@ public:
 	virtual ~ISparseSampledSequenceTransportCoordinator() = default;
 
 	virtual const TRange<double> GetDisplayRange() const = 0;
-	virtual const float GetFocusPoint() const = 0;
+	virtual const double GetFocusPoint() const = 0;
 	virtual void ScrubFocusPoint(const float InTargetFocusPoint, const bool bIsMoving) = 0;
 	virtual const bool IsScrubbing() const = 0;
 
diff --git a/Engine/Plugins/Runtime/AudioWidgets/Source/AudioWidgets/Public/SFixedSampledSequenceRuler.h b/Engine/Plugins/Runtime/AudioWidgets/Source/AudioWidgets/Public/SFixedSampledSequenceRuler.h
--- a/Engine/Plugins/Runtime/AudioWidgets/Source/AudioWidgets/Public/SFixedSampledSequenceRuler.h
+++ b/Engine/Plugins/Runtime/AudioWidgets/Source/AudioWidgets/Public/SFixedSampledSequenceRuler.h
@@ -38,6 +38,8 @@ public:
 	UE_API void OnStyleUpdated(const FFixedSampleSequenceRulerStyle UpdatedStyle);
 	UE_API FReply LaunchContextMenu();
 
+	UE_API void SetPlayheadVisible(bool Enable);
+
 	/** Delegate sent when the user selects a new display unit from the RMB menu*/
 	FOnTimeUnitMenuSelection OnTimeUnitMenuSelection;
 
diff --git a/Engine/Plugins/Runtime/AudioWidgets/Source/AudioWidgets/Public/SPlayheadOverlay.h b/Engine/Plugins/Runtime/AudioWidgets/Source/AudioWidgets/Public/SPlayheadOverlay.h
--- a/Engine/Plugins/Runtime/AudioWidgets/Source/AudioWidgets/Public/SPlayheadOverlay.h
+++ b/Engine/Plugins/Runtime/AudioWidgets/Source/AudioWidgets/Public/SPlayheadOverlay.h
@@ -24,6 +24,8 @@ public:
 	UE_API void SetPlayheadPosition(const float InNewPosition);
 	UE_API void OnStyleUpdated(const FPlayheadOverlayStyle UpdatedStyle);
 
+	UE_API void SetPlayheadVisible(bool Enable);
+
 private:
 	UE_API virtual int32 OnPaint(const FPaintArgs& Args, const FGeometry& AllottedGeometry, const FSlateRect& MyCullingRect, FSlateWindowElementList& OutDrawElements, int32 LayerId, const FWidgetStyle& InWidgetStyle, bool bParentEnabled) const override;
 	UE_API virtual FVector2D ComputeDesiredSize(float) const override;
@@ -38,6 +40,7 @@ private:
 	float DesiredHeight = 0.f;
 
 	float PlayheadPosition = 0.f;
+	bool bIsVisible = true;
 };
 
 #undef UE_API
diff --git a/Engine/Plugins/Runtime/AudioWidgets/Source/AudioWidgets/Public/SparseSampledSequenceTransportCoordinator.h b/Engine/Plugins/Runtime/AudioWidgets/Source/AudioWidgets/Public/SparseSampledSequenceTransportCoordinator.h
--- a/Engine/Plugins/Runtime/AudioWidgets/Source/AudioWidgets/Public/SparseSampledSequenceTransportCoordinator.h
+++ b/Engine/Plugins/Runtime/AudioWidgets/Source/AudioWidgets/Public/SparseSampledSequenceTransportCoordinator.h
@@ -8,7 +8,7 @@
 #define UE_API AUDIOWIDGETS_API
 
 DECLARE_MULTICAST_DELEGATE_OneParam(FOnDisplayRangeUpdated, const TRange<double> /* New Display Range */);
-DECLARE_MULTICAST_DELEGATE_OneParam(FOnFocusPointMoved, const float /* New Position */);
+DECLARE_MULTICAST_DELEGATE_OneParam(FOnFocusPointMoved, const double /* New Position */);
 DECLARE_MULTICAST_DELEGATE_TwoParams(FOnFocusPointScrubUpdate, const float /* Focused Playback Ratio */, const bool /*Playhead is Moving*/);
 
 class FSparseSampledSequenceTransportCoordinator : public ISparseSampledSequenceTransportCoordinator
@@ -28,7 +28,7 @@ public:
 
 	/** ISparseSampledSequenceTransportCoordinator interface */
 	UE_API const TRange<double> GetDisplayRange() const;
-	UE_API const float GetFocusPoint() const override;
+	UE_API const double GetFocusPoint() const override;
 	UE_API void ScrubFocusPoint(const float InTargetFocusPoint, const bool bIsMoving) override;
 	UE_API const bool IsScrubbing() const;
 	UE_API void SetProgressRatio(const float NewRatio) override;
@@ -48,7 +48,6 @@ private:
 	UE_API void UpdateDisplayRange(const double MinValue, const double MaxValue);
 	UE_API bool IsRatioWithinDisplayRange(const double Ratio) const;
 
-	double CurrentPlaybackRatio = 0.f;
 	double FocusPointLockPosition = 0.95f;
 	double FocusPoint = 0;
 	float ZoomRatio = 1.f;
