[Waveform Editor] Audio.WaveformEditor Tests failing due to TestResetPlayhead and legacy assets being altered by the waveform editor (as intended).
Changed TestResetPlayhead to use the WaveformView frame position instead of TransportCoordinator->GetFocusPoint().
Updated stereo asset so on open, waveform editor no longer needs to alter it.

#virtualized
#Synced-CL 51885885

diff --git a/Engine/Plugins/Editor/WaveformEditor/Source/WaveformEditor/Public/WaveformEditor.h b/Engine/Plugins/Editor/WaveformEditor/Source/WaveformEditor/Public/WaveformEditor.h
--- a/Engine/Plugins/Editor/WaveformEditor/Source/WaveformEditor/Public/WaveformEditor.h
+++ b/Engine/Plugins/Editor/WaveformEditor/Source/WaveformEditor/Public/WaveformEditor.h
@@ -59,6 +59,7 @@ public:
 	UE_API void ResetPlaybackToStart();
 
 	TSharedPtr<class FSparseSampledSequenceTransportCoordinator> GetTransportCoordinator() { return TransportCoordinator; }
+	const int64 GetPlayheadFramePosition() const { return WaveformView.ViewWidget.IsValid() ? WaveformView.ViewWidget->GetPlayheadFramePosition() : 0; }
 
 	static UE_API const double EndTimeInvalid;
 
