[Waveform Editor] Unit test: 0-length .wav import/support.
Also prevent the Waveform Editor from programmatically opening a 0-length soundwave and crashing.
Background: In older versions of UE, users would import zero-length .wav files as a workaround for playing subtitles without accompanying audio.
#virtualized
#Synced-CL 51945513

diff --git a/Engine/Plugins/Editor/WaveformEditor/Source/WaveformEditor/Private/WaveformEditor.cpp b/Engine/Plugins/Editor/WaveformEditor/Source/WaveformEditor/Private/WaveformEditor.cpp
--- a/Engine/Plugins/Editor/WaveformEditor/Source/WaveformEditor/Private/WaveformEditor.cpp
+++ b/Engine/Plugins/Editor/WaveformEditor/Source/WaveformEditor/Private/WaveformEditor.cpp
@@ -57,55 +57,72 @@ bool FWaveformEditor::Init(const EToolkitMode::Type Mode, const TSharedPtr<ITool
 
 	SoundWave = SoundWaveToEdit;
 
-	// Initialize TransformationChainConfig from persisted transformations or if there are none falling back to default transformations.
-	if (SoundWave->Transformations.Num() == 0)
+	bool bIsInitialized = false;
+
+	// If either condition is zero it is an indication of a zero length or corrupt file.
+	// FWaveformEditorInstantiator::CanSoundWaveBeOpenedInEditor already prevents zero-duration waves from opening through the normal route, added guard here to block in other scenarios (e.g. automated testing)
+	if (FMath::IsNearlyZero(SoundWave->GetDuration()) || FMath::IsNearlyZero(SoundWave->TotalSamples))
 	{
-		AddDefaultTransformations();
+		UE_LOG(LogWaveformEditor, Warning, TEXT("Waveform Editor attempted to open Zero-Length SoundWave: %s"), *SoundWave->GetName());
+
+		bIsInitialized = false;
 	}
 	else
 	{
-		TransformationChainConfig = SoundWave->UpdateTransformations();
-	}
-
-	// Add a trimfade transformation for user to opt to use, but will be removed from transformations on save if values are zero
-	// This does not trigger a modify on the soundwave to avoid marking dirty every time a user opens the waveform, but if they change anything on the trimfade that automatically will modify and mark dirty to save it
-	TObjectPtr<UWaveformTransformationTrimFade> TrimFadeTransformation = GetOrAddTrimFadeTransformation(false);
-	check(TrimFadeTransformation != nullptr);
+		bIsInitialized = true;
+
+		bIsInitialized &= CreateDetailsViews();
+		bIsInitialized &= CreateTransportCoordinator();
+		bIsInitialized &= InitializeZoom();
+		bIsInitialized &= CreateWaveformView();
+		bIsInitialized &= InitializeAudioComponent();
+		bIsInitialized &= CreateTransportController();
+		bIsInitialized &= CreateWaveWriter();
+		bIsInitialized &= BindDelegates();
+		bIsInitialized &= SetUpAssetReimport();
+
+		bIsInitialized &= RegisterToolbar();
+		bIsInitialized &= BindCommands();
+	}
+
+	// Must be initialized whether init succeeds or fails for proper destruction
+	FAssetEditorToolkit::InitAssetEditor(
+		Mode,
+		InitToolkitHost,
+		AppIdentifier,
+		StandaloneDefaultLayout,
+		bCreateDefaultStandaloneMenu,
+		bCreateDefaultToolbar,
+		SoundWaveToEdit,
+		bToolbarFocusable,
+		bUseSmallIcons);
 
-	if (FMath::IsNearlyEqual(TrimFadeTransformation->EndTime, EndTimeInvalid))
+	if (bIsInitialized)
 	{
-		TrimFadeTransformation->EndTime = SoundWave->Duration;
-	}
-
-	bool bIsInitialized = true;
-	
-	bIsInitialized &= CreateDetailsViews();
-	bIsInitialized &= CreateTransportCoordinator();
-	bIsInitialized &= InitializeZoom();
-	bIsInitialized &= CreateWaveformView();
-	bIsInitialized &= InitializeAudioComponent();
-	bIsInitialized &= CreateTransportController();
-	bIsInitialized &= CreateWaveWriter();
-	bIsInitialized &= BindDelegates();
-	bIsInitialized &= SetUpAssetReimport();
+		if (GEditor)
+		{
+			GEditor->RegisterForUndo(this);
+		}
 
-	bIsInitialized &= RegisterToolbar();
-	bIsInitialized &= BindCommands();
+		// Initialize TransformationChainConfig from persisted transformations or if there are none falling back to default transformations.
+		if (SoundWave->Transformations.Num() == 0)
+		{
+			AddDefaultTransformations();
+		}
+		else
+		{
+			TransformationChainConfig = SoundWave->UpdateTransformations();
+		}
 
-	GEditor->RegisterForUndo(this);
+		// Add a trimfade transformation for user to opt to use, but will be removed from transformations on save if values are zero
+		// This does not trigger a modify on the soundwave to avoid marking dirty every time a user opens the waveform, but if they change anything on the trimfade that automatically will modify and mark dirty to save it
+		TObjectPtr<UWaveformTransformationTrimFade> TrimFadeTransformation = GetOrAddTrimFadeTransformation(false);
+		check(TrimFadeTransformation != nullptr);
 
-	if (bIsInitialized)
-	{
-		FAssetEditorToolkit::InitAssetEditor(
-			Mode,
-			InitToolkitHost,
-			AppIdentifier,
-			StandaloneDefaultLayout,
-			bCreateDefaultStandaloneMenu,
-			bCreateDefaultToolbar,
-			SoundWaveToEdit,
-			bToolbarFocusable,
-			bUseSmallIcons);
+		if (FMath::IsNearlyEqual(TrimFadeTransformation->EndTime, EndTimeInvalid))
+		{
+			TrimFadeTransformation->EndTime = SoundWave->Duration;
+		}
 
 		// Initialize waveform view to with playhead at start of active transformations.
 		check(WaveformView.DataProvider != nullptr);
@@ -160,10 +177,10 @@ bool FWaveformEditor::Init(const EToolkitMode::Type Mode, const TSharedPtr<ITool
 			}
 		}
 #endif // WITH_EDITOR
+	
+		FSlateApplication::Get().SetKeyboardFocus(TransformationsDetails);
 	}
 
-	FSlateApplication::Get().SetKeyboardFocus(TransformationsDetails);
-
 	return bIsInitialized;
 }
 
@@ -233,6 +250,11 @@ FWaveformEditor::~FWaveformEditor()
 	{
 		FReimportManager::Instance()->OnPostReimport().RemoveAll(this);
 	}
+
+	if (GEditor)
+	{
+		GEditor->UnregisterForUndo(this);
+	}
 }
 
 void FWaveformEditor::AddDefaultTransformations()
