WaveformEditor: Fix crash when deleting an open SoundWave asset with the editor open/minimized

Guard SoundWave modification in ~FWaveformEditor() when the editor is
closed due to asset deletion, preventing check() failures in the
TrimFade and loop-region cleanup blocks.

#okforgitub public
#Synced-CL 52025567

diff --git a/Engine/Plugins/Editor/WaveformEditor/Source/WaveformEditor/Private/WaveformEditor.cpp b/Engine/Plugins/Editor/WaveformEditor/Source/WaveformEditor/Private/WaveformEditor.cpp
--- a/Engine/Plugins/Editor/WaveformEditor/Source/WaveformEditor/Private/WaveformEditor.cpp
+++ b/Engine/Plugins/Editor/WaveformEditor/Source/WaveformEditor/Private/WaveformEditor.cpp
@@ -184,6 +184,18 @@ bool FWaveformEditor::Init(const EToolkitMode::Type Mode, const TSharedPtr<ITool
 	return bIsInitialized;
 }
 
+bool FWaveformEditor::CloseWindow(EAssetEditorCloseReason InCloseReason)
+{
+	if (InCloseReason == EAssetEditorCloseReason::CloseAllEditorsForAsset
+		|| InCloseReason == EAssetEditorCloseReason::AssetUnloadingOrInvalid
+		|| InCloseReason == EAssetEditorCloseReason::AssetForceDeleted)
+	{
+		bSoundWaveBeingDeleted = true;
+	}
+
+	return FAssetEditorToolkit::CloseWindow(InCloseReason);
+}
+
 FWaveformEditor::~FWaveformEditor()
 {
 	if (ChangeCuePointOriginNotificationItem.IsValid())
@@ -201,7 +213,8 @@ FWaveformEditor::~FWaveformEditor()
 			MarkerTransformation->Markers->CueChanged.Unbind();
 		}
 
-		if (TransformationChainConfig.bIsPreviewingLoopRegion)
+		// Skip modifying the SoundWave if it is being deleted — modifications are invalid during asset deletion
+		if (!bSoundWaveBeingDeleted && TransformationChainConfig.bIsPreviewingLoopRegion)
 		{
 			// Note: if the user doesn't save the soundwave asset, bLooping will not be reset.
 			check(SoundWave);
@@ -215,7 +228,7 @@ FWaveformEditor::~FWaveformEditor()
 		}
 	}
 
-	TObjectPtr<UWaveformTransformationTrimFade> TrimFadeTransformation = GetTrimFadeTransformation();
+	TObjectPtr<UWaveformTransformationTrimFade> TrimFadeTransformation = bSoundWaveBeingDeleted ? nullptr : GetTrimFadeTransformation();
 
 	if (TrimFadeTransformation)
 	{
diff --git a/Engine/Plugins/Editor/WaveformEditor/Source/WaveformEditor/Public/WaveformEditor.h b/Engine/Plugins/Editor/WaveformEditor/Source/WaveformEditor/Public/WaveformEditor.h
--- a/Engine/Plugins/Editor/WaveformEditor/Source/WaveformEditor/Public/WaveformEditor.h
+++ b/Engine/Plugins/Editor/WaveformEditor/Source/WaveformEditor/Public/WaveformEditor.h
@@ -31,6 +31,7 @@ public:
 	UE_API virtual ~FWaveformEditor();
 
 	/** FAssetEditorToolkit interface */
+	UE_API virtual bool CloseWindow(EAssetEditorCloseReason InCloseReason) override;
 	UE_API virtual void RegisterTabSpawners(const TSharedRef<FTabManager>& TabManager) override;
 	UE_API virtual void UnregisterTabSpawners(const TSharedRef<FTabManager>& TabManager) override;
 	UE_API virtual FName GetEditorName() const override;
@@ -220,6 +221,10 @@ private:
 	TSharedPtr<SNotificationItem> ChangeCuePointOriginNotificationItem = nullptr;
 	
 	bool bMarkerTransformationIsActive = false;
+
+	/** Set to true when the editor is being closed because the SoundWave asset is being deleted or invalidated.
+	 *  Guards against unsafe SoundWave modification in the destructor cleanup path. */
+	bool bSoundWaveBeingDeleted = false;
 };
 
 #undef UE_API
