[SoundWave] - Opening a Sound Wave asset as it plays in editor stops the sound from playing
USoundWave::PostEditChange() calls UpdateAsset() which kills all playback of the sound asset. Using PostEditChangeProperty() with a valid property event can avoid this for most properties. For Transformations, added a stripped down version of UpdateAsset() that doesn't stop the playback of the asset.

#Synced-CL 52100586

diff --git a/Engine/Plugins/Editor/WaveformEditor/Source/WaveformEditor/Private/WaveformEditor.cpp b/Engine/Plugins/Editor/WaveformEditor/Source/WaveformEditor/Private/WaveformEditor.cpp
--- a/Engine/Plugins/Editor/WaveformEditor/Source/WaveformEditor/Private/WaveformEditor.cpp
+++ b/Engine/Plugins/Editor/WaveformEditor/Source/WaveformEditor/Private/WaveformEditor.cpp
@@ -224,7 +224,7 @@ FWaveformEditor::~FWaveformEditor()
 			MarkerTransformation->ResetLoopPreviewing();
 			SoundWave->UpdateTransformations();
 
-			SoundWave->PostEditChange();
+			NotifySoundWavePostEditChangeProperty(GET_MEMBER_NAME_CHECKED(USoundWave, bLooping));
 		}
 	}
 
@@ -311,7 +311,7 @@ void FWaveformEditor::AddDefaultTransformations()
 			}
 
 			//Update the content browser asset state to dirty
-			SoundWave->PostEditChange();
+			NotifySoundWavePostEditChangeProperty(GET_MEMBER_NAME_CHECKED(USoundWave, Transformations), EPropertyChangeType::ArrayAdd);
 		}
 
 		TransformationChainConfig = SoundWave->UpdateTransformations();
@@ -333,7 +333,7 @@ void FWaveformEditor::NotifyPostTransformationChange(const EPropertyChangeType::
 	NotifyPostChange(PropertyChangedEvent, &PropertyChain);
 
 	//Update the content browser asset state to dirty
-	SoundWave->PostEditChange();
+	NotifySoundWavePostEditChangeProperty(GET_MEMBER_NAME_CHECKED(USoundWave, Transformations), PropertyChangeType);
 
 	TransformationChainConfig = SoundWave->UpdateTransformations();
 }
@@ -812,7 +812,7 @@ void FWaveformEditor::UpdateTransformations(bool bMarkDirty)
 	bool bWasAlreadyDirty = SoundWave->GetPackage()->IsDirty();
 	
 	//Update the content browser asset state to dirty
-	SoundWave->PostEditChange();
+	NotifySoundWavePostEditChangeProperty(GET_MEMBER_NAME_CHECKED(USoundWave, Transformations));
 	TransformationChainConfig = SoundWave->UpdateTransformations();
 
 	FProperty* TransformationsProperty = FindFProperty<FProperty>(USoundWave::StaticClass(), GET_MEMBER_NAME_CHECKED(USoundWave, Transformations));
@@ -852,6 +852,18 @@ void FWaveformEditor::UpdateRenderElements()
 	TransformationChainConfig = SoundWave->UpdateTransformations();
 }
 
+void FWaveformEditor::NotifySoundWavePostEditChangeProperty(const FName& PropertyName, const EPropertyChangeType::Type PropertyChangeType)
+{
+	check(SoundWave);
+
+	FProperty* Property = USoundWave::StaticClass()->FindPropertyByName(PropertyName);
+	check(Property);
+
+	FPropertyChangedEvent PropertyChangedEvent(Property, PropertyChangeType);
+	SoundWave->PostEditChangeProperty(PropertyChangedEvent);
+	SoundWave->MarkPackageDirty();
+}
+
 void FWaveformEditor::ModifyMarkerLoopRegion(ELoopModificationControls Modification)
 {
 	check(SoundWave);
@@ -888,7 +900,7 @@ void FWaveformEditor::SetCuePointOrigin(TObjectPtr<UWaveformTransformationMarker
 			GEditor->BeginTransaction(TEXT("PropertyEditor"), LOCTEXT("SettingCuePointOrigin", "SetCuePointOrigin"), SoundWave);
 			SoundWave->Modify();
 			SoundWave->SetCuePointOrigin(ESoundWaveCuePointOrigin::MarkerTransformation);
-			SoundWave->PostEditChange();
+			NotifySoundWavePostEditChangeProperty(USoundWave::GetCuePointOriginPropertyName());
 			GEditor->EndTransaction();
 		}
 		else if (SoundWave->GetCuePointOrigin() != ESoundWaveCuePointOrigin::WaveFile)
@@ -896,7 +908,7 @@ void FWaveformEditor::SetCuePointOrigin(TObjectPtr<UWaveformTransformationMarker
 			GEditor->BeginTransaction(TEXT("PropertyEditor"), LOCTEXT("SettingCuePointOrigin", "SetCuePointOrigin"), SoundWave);
 			SoundWave->Modify();
 			SoundWave->SetCuePointOrigin(ESoundWaveCuePointOrigin::WaveFile);
-			SoundWave->PostEditChange();
+			NotifySoundWavePostEditChangeProperty(USoundWave::GetCuePointOriginPropertyName());
 			GEditor->EndTransaction();
 		}
 	}
@@ -1948,7 +1960,7 @@ void FWaveformEditor::UpdateSoundWaveProperties()
 	SoundWave->SetLoudnessValues(LUFS, SamplePeakDB);
 
 	// Update SoundWave tooltip values to show new loudness values in the content browser when hovering over the sound
-	SoundWave->PostEditChange();
+	NotifySoundWavePostEditChangeProperty(GET_MEMBER_NAME_CHECKED(USoundWave, LUFS));
 }
 #endif //WITH_EDITOR
 
diff --git a/Engine/Plugins/Editor/WaveformEditor/Source/WaveformEditor/Public/WaveformEditor.h b/Engine/Plugins/Editor/WaveformEditor/Source/WaveformEditor/Public/WaveformEditor.h
--- a/Engine/Plugins/Editor/WaveformEditor/Source/WaveformEditor/Public/WaveformEditor.h
+++ b/Engine/Plugins/Editor/WaveformEditor/Source/WaveformEditor/Public/WaveformEditor.h
@@ -147,6 +147,10 @@ private:
 	UE_API void UpdateSoundWaveProperties(); 
 #endif //WITH_EDITOR
 
+	// SoundWave->PostEditChange() calls UpdateAsset() which kills all playback of the sound asset.
+	// Calling SoundWave->PostEditChangeProperty(PropertyChangedEvent) with a valid property event avoids this.
+	UE_API void NotifySoundWavePostEditChangeProperty(const FName& PropertyName, const EPropertyChangeType::Type PropertyChangeType = EPropertyChangeType::ValueSet);
+
 	UE_API void ModifyMarkerLoopRegion(ELoopModificationControls Modification);
 	UE_API void CycleMarkerLoopRegion(ELoopModificationControls Modification);
 
diff --git a/Engine/Source/Runtime/Engine/Classes/Sound/SoundWave.h b/Engine/Source/Runtime/Engine/Classes/Sound/SoundWave.h
--- a/Engine/Source/Runtime/Engine/Classes/Sound/SoundWave.h
+++ b/Engine/Source/Runtime/Engine/Classes/Sound/SoundWave.h
@@ -497,6 +497,8 @@ public:
 	ENGINE_API const ESoundWaveCuePointOrigin GetCuePointOrigin() const;
 	ENGINE_API void SetCuePointOrigin(const ESoundWaveCuePointOrigin InCuePointOrigin);
 
+	static FName GetCuePointOriginPropertyName() { return GET_MEMBER_NAME_CHECKED(USoundWave, CuePointOrigin); }
+
 	/** Set gated LUFS and SamplePeakDB values */
 	ENGINE_API void SetLoudnessValues(const float InLUFS, const float InSamplePeakDB);
 #endif //WITH_EDITOR
diff --git a/Engine/Source/Runtime/Engine/Private/SoundWave.cpp b/Engine/Source/Runtime/Engine/Private/SoundWave.cpp
--- a/Engine/Source/Runtime/Engine/Private/SoundWave.cpp
+++ b/Engine/Source/Runtime/Engine/Private/SoundWave.cpp
@@ -3356,13 +3356,27 @@ void USoundWave::PostEditChangeProperty(FPropertyChangedEvent& PropertyChangedEv
 				bSoundWaveDataNeedsInitialization = false; //< SetSoundAssetCompressionType eventually calls InitializeDataFromSoundWave(...)
 			}
 
-			if (Name == CompressionQualityFName
+			if (Name == TransformationsFName)
+			{
+				// For transformation changes, invalidate and rebuild compressed data without
+				// stopping currently playing sounds. InvalidateCompressedData forks the
+				// SoundWaveData so existing proxies (held by playing sounds) remain valid
+				// while new playback requests will pick up the updated data.
+				// Essentially UpdateAsset() without FreeResources() and FlushAudioRenderingCommands();
+
+				InvalidateCompressedData();
+				UpdatePlatformData();
+				// Reset any error state we've encountered.
+				ResetError();
+
+				bSoundWaveDataNeedsInitialization = false; //< InvalidateCompressedData() eventually calls InitializeDataFromSoundWave(...)
+			}
+			else if (Name == CompressionQualityFName
 				|| Name == SampleRateFName
 				|| Name == StreamingFName
 				|| Name == LoadingBehaviorFName
 				|| Name == InlinedAudioInSecondsFName
 				|| Name == InitialChunkSizeFName
-				|| Name == TransformationsFName
 				|| Name == CloudStreamingPlatformSettingsFName
 				|| Name == CloudStreamingFName)
 			{
