[WaveformEditor] Fix Sound Waves being dirtied when previewed in Content Browser

The WaveformEditor module hooks OnSoundBasePreview and runs loudness analysis on legacy sound waves (LUFS == 0). After computing values, it
calls PostEditChange() which triggers MarkPackageDirty(). This is only
appropriate on import, not on preview.
Changed to calling PostEditChangeProperty() with a valid property event, then calling MarkPackageDirty() optionally.

Added commandlet to update all legacy USoundWaves with the new LUFS and SamplePeakDB data.

Example use command: UnrealEditor.exe EngineTest.uproject -run=ResaveSoundWaveLoudness -PACKAGEFOLDER=/Game/Audio -AutoCheckOut

#Synced-CL 52135216

diff --git a/Engine/Plugins/Editor/WaveformEditor/Source/WaveformEditor/Private/ResaveSoundWaveLoudnessCommandlet.cpp b/Engine/Plugins/Editor/WaveformEditor/Source/WaveformEditor/Private/ResaveSoundWaveLoudnessCommandlet.cpp
new file mode 100644
--- /dev/null
+++ b/Engine/Plugins/Editor/WaveformEditor/Source/WaveformEditor/Private/ResaveSoundWaveLoudnessCommandlet.cpp
@@ -0,0 +1,105 @@
+// Copyright Epic Games, Inc. All Rights Reserved.
+
+#include "ResaveSoundWaveLoudnessCommandlet.h"
+
+#include "AssetRegistry/AssetData.h"
+#include "AssetRegistry/AssetRegistryModule.h"
+#include "Misc/PackageName.h"
+#include "Sound/SoundWave.h"
+#include "Sound/SoundSourceBus.h"
+#include "Sound/SoundWaveProcedural.h"
+#include "WaveformEditorLog.h"
+#include "WaveformEditorModule.h"
+
+UResaveSoundWaveLoudnessCommandlet::UResaveSoundWaveLoudnessCommandlet()
+{
+	// Shown in help text when running -run=help
+	HelpDescription = TEXT("Migrates legacy SoundWave assets by computing missing loudness metadata (LUFS and SamplePeakDB).");
+	HelpUsage = TEXT("UnrealEditor.exe Project.uproject -run=ResaveSoundWaveLoudness [-PACKAGEFOLDER=/Game/Audio] [-AutoCheckOut] [-AutoCheckIn]");
+
+	// We only care about saving packages that we actually modify.
+	bOnlySaveDirtyPackages = true;
+}
+
+int32 UResaveSoundWaveLoudnessCommandlet::Main(const FString& Params)
+{
+	// Run the base ResavePackages loop, which calls PerformAdditionalOperations for each object.
+	const int32 Result = Super::Main(Params);
+
+	// Report migration summary.
+	UE_LOG(LogWaveformEditor, Display, TEXT(""));
+	UE_LOG(LogWaveformEditor, Display, TEXT("=== SoundWave Loudness Migration Summary ==="));
+	UE_LOG(LogWaveformEditor, Display, TEXT("  Migrated : %d (loudness computed and package resaved)"), NumMigrated);
+	UE_LOG(LogWaveformEditor, Display, TEXT("  Skipped  : %d (already had loudness data)"), NumSkipped);
+	UE_LOG(LogWaveformEditor, Display, TEXT("  Failed   : %d (could not read raw PCM data)"), NumFailed);
+	UE_LOG(LogWaveformEditor, Display, TEXT("  Total    : %d"), NumMigrated + NumSkipped + NumFailed);
+	UE_LOG(LogWaveformEditor, Display, TEXT("============================================="));
+
+	return Result;
+}
+
+bool UResaveSoundWaveLoudnessCommandlet::ShouldSkipPackage(const FString& Filename)
+{
+	if (Super::ShouldSkipPackage(Filename))
+	{
+		return true;
+	}
+
+	FString PackageName;
+	if (!FPackageName::TryConvertFilenameToLongPackageName(Filename, PackageName))
+	{
+		return true;
+	}
+
+	IAssetRegistry& AssetRegistry = FModuleManager::LoadModuleChecked<FAssetRegistryModule>(TEXT("AssetRegistry")).Get();
+
+	TArray<FAssetData> Assets;
+	AssetRegistry.GetAssetsByPackageName(FName(*PackageName), Assets, true);
+
+	for (const FAssetData& Asset : Assets)
+	{
+		if (Asset.IsInstanceOf(USoundWave::StaticClass())
+			&& !Asset.IsInstanceOf(USoundWaveProcedural::StaticClass())
+			&& !Asset.IsInstanceOf(USoundSourceBus::StaticClass()))
+		{
+			return false;
+		}
+	}
+
+	return true;
+}
+
+void UResaveSoundWaveLoudnessCommandlet::PerformAdditionalOperations(UObject* Object, bool& bSavePackage)
+{
+	Super::PerformAdditionalOperations(Object, bSavePackage);
+
+	USoundWave* SoundWave = Cast<USoundWave>(Object);
+
+	if (!SoundWave)
+	{
+		return;
+	}
+
+	// Filtered out by UResaveSoundWaveLoudnessCommandlet::ShouldSkipPackage
+	check(!SoundWave->IsA<USoundWaveProcedural>() && !SoundWave->IsA<USoundSourceBus>());
+
+	// Only migrate assets that are missing loudness data.
+	if (SoundWave->LUFS != 0.f || SoundWave->SamplePeakDB != 0.f)
+	{
+		NumSkipped++;
+		return;
+	}
+
+	// Compute and apply loudness values using the shared analysis function.
+	if (!WaveformAnalysis::AnalyzeSoundWaveLoudness(SoundWave))
+	{
+		UE_LOG(LogWaveformEditor, Warning, TEXT("  FAILED: '%s' - could not analyze raw PCM data"), *SoundWave->GetPathName());
+		NumFailed++;
+		return;
+	}
+
+	bSavePackage = true;
+
+	UE_LOG(LogWaveformEditor, Display, TEXT("  Migrated: '%s' (LUFS=%.1f, PeakDB=%.1f)"), *SoundWave->GetPathName(), SoundWave->LUFS, SoundWave->SamplePeakDB);
+	NumMigrated++;
+}
diff --git a/Engine/Plugins/Editor/WaveformEditor/Source/WaveformEditor/Private/ResaveSoundWaveLoudnessCommandlet.h b/Engine/Plugins/Editor/WaveformEditor/Source/WaveformEditor/Private/ResaveSoundWaveLoudnessCommandlet.h
new file mode 100644
--- /dev/null
+++ b/Engine/Plugins/Editor/WaveformEditor/Source/WaveformEditor/Private/ResaveSoundWaveLoudnessCommandlet.h
@@ -0,0 +1,50 @@
+// Copyright Epic Games, Inc. All Rights Reserved.
+
+#pragma once
+
+#include "Commandlets/ResavePackagesCommandlet.h"
+
+#include "ResaveSoundWaveLoudnessCommandlet.generated.h"
+
+/**
+ * Commandlet to migrate legacy USoundWave assets that are missing loudness metadata
+ * (LUFS and SamplePeakDB). These values were not computed on import prior to 5.8.
+ *
+ * Inherits from UResavePackagesCommandlet, so all standard path filtering and source
+ * control options are available:
+ *
+ *   # Migrate all SoundWaves in the project
+ *   UnrealEditor.exe Project.uproject -run=ResaveSoundWaveLoudness -AutoCheckOut
+ *
+ *   # Scope to specific content folders (+ delimited for multiple paths)
+ *   UnrealEditor.exe Project.uproject -run=ResaveSoundWaveLoudness -PACKAGEFOLDER=/Game/Audio+/Game/SFX -AutoCheckOut
+ *
+ *   # With batched source control and auto-checkin
+ *   UnrealEditor.exe Project.uproject -run=ResaveSoundWaveLoudness -AutoCheckOut -AutoCheckIn -BatchSourceControl
+ */
+UCLASS()
+class UResaveSoundWaveLoudnessCommandlet : public UResavePackagesCommandlet
+{
+	GENERATED_BODY()
+
+public:
+	UResaveSoundWaveLoudnessCommandlet();
+
+	// UCommandlet interface
+	virtual int32 Main(const FString& Params) override;
+
+protected:
+	// UResavePackagesCommandlet interface
+	virtual bool ShouldSkipPackage(const FString& Filename) override;
+	virtual void PerformAdditionalOperations(UObject* Object, bool& bSavePackage) override;
+
+private:
+	/** Number of SoundWaves that had loudness data computed and will be resaved. */
+	int32 NumMigrated = 0;
+
+	/** Number of SoundWaves that already had valid loudness data. */
+	int32 NumSkipped = 0;
+
+	/** Number of SoundWaves where loudness analysis failed (no raw PCM data). */
+	int32 NumFailed = 0;
+};
diff --git a/Engine/Plugins/Editor/WaveformEditor/Source/WaveformEditor/Private/WaveformEditorModule.cpp b/Engine/Plugins/Editor/WaveformEditor/Source/WaveformEditor/Private/WaveformEditorModule.cpp
--- a/Engine/Plugins/Editor/WaveformEditor/Source/WaveformEditor/Private/WaveformEditorModule.cpp
+++ b/Engine/Plugins/Editor/WaveformEditor/Source/WaveformEditor/Private/WaveformEditorModule.cpp
@@ -18,9 +18,10 @@
 
 DEFINE_LOG_CATEGORY(LogWaveformEditor);
 
-namespace
+namespace WaveformAnalysis
 {
-	void AnalyzeSoundWaveLoudness(UObject* Object)
+	// Analyzes loudness of a sound wave and sets LUFS/SamplePeakDB values.
+	bool AnalyzeSoundWaveLoudness(UObject* Object, bool bMarkDirty)
 	{
 		if (Object && Object->IsA<USoundWave>() && !Object->IsA<USoundWaveProcedural>())
 		{
@@ -34,9 +35,9 @@ namespace
 
 			if (!SoundWave->GetImportedSoundWaveData(ImportedRawPCMData, BufferSampleRate, ImportedNumChannels) || BufferSampleRate == 0)
 			{
-				UE_LOG(LogWaveformEditor, Warning, TEXT("Soundwave loudness analyzation failed in Waveform Editor Module."));
+				UE_LOG(LogWaveformEditor, Warning, TEXT("Soundwave loudness analyzation failed, could not import raw PCM data."));
 
-				return;
+				return false;
 			}
 
 			const int64 NumWaveformSamples = ImportedRawPCMData.Num() * sizeof(uint8) / sizeof(int16);
@@ -47,8 +48,24 @@ namespace
 			const float LUFS = WaveformAudioAnalysis::GetLUFS(AudioBuffer, BufferSampleRate, ImportedNumChannels);
 
 			SoundWave->SetLoudnessValues(LUFS, SamplePeakDB);
-			SoundWave->PostEditChange();
+
+			// Updates the content browser so the tooltip displays updated metadata, uses PostEditChangeProperty() to avoid marking dirty
+			FProperty* Property = USoundWave::StaticClass()->FindPropertyByName(GET_MEMBER_NAME_CHECKED(USoundWave, LUFS));
+			FPropertyChangedEvent PropertyChangedEvent(Property, EPropertyChangeType::ValueSet);
+			SoundWave->PostEditChangeProperty(PropertyChangedEvent);
+
+			// Optionally mark the file dirty
+			if (bMarkDirty)
+			{
+				SoundWave->MarkPackageDirty();
+			}
 		}
+		else 
+		{
+			return false;
+		}
+
+		return true;
 	}
 }
 
@@ -69,7 +86,7 @@ void FWaveformEditorModule::StartupModule()
 				{
 					OnAssetPostImportHandle = ImportSubsystem->OnAssetPostImport.AddLambda([&](UFactory* Factory, UObject* Object)
 						{
-							AnalyzeSoundWaveLoudness(Object);
+							WaveformAnalysis::AnalyzeSoundWaveLoudness(Object);
 						});
 				}
 
@@ -81,7 +98,8 @@ void FWaveformEditorModule::StartupModule()
 
 							if (SoundWave->LUFS == 0.f && SoundWave->SamplePeakDB == 0.f)
 							{
-								AnalyzeSoundWaveLoudness(SoundBase);
+								// Pass bMarkDirty=false: previewing a sound should not dirty the package.
+								WaveformAnalysis::AnalyzeSoundWaveLoudness(SoundBase, /*bMarkDirty=*/ false);
 							}
 						}
 					});
@@ -106,7 +124,7 @@ void FWaveformEditorModule::StartupModule()
 
 						for (UObject* Object : ObjectsInPackage)
 						{
-							AnalyzeSoundWaveLoudness(Object);
+							WaveformAnalysis::AnalyzeSoundWaveLoudness(Object);
 						}
 					}
 				});
diff --git a/Engine/Plugins/Editor/WaveformEditor/Source/WaveformEditor/Public/WaveformEditorModule.h b/Engine/Plugins/Editor/WaveformEditor/Source/WaveformEditor/Public/WaveformEditorModule.h
--- a/Engine/Plugins/Editor/WaveformEditor/Source/WaveformEditor/Public/WaveformEditorModule.h
+++ b/Engine/Plugins/Editor/WaveformEditor/Source/WaveformEditor/Public/WaveformEditorModule.h
@@ -8,6 +8,14 @@
 
 class IWaveformEditorInstantiator;
 
+namespace WaveformAnalysis
+{
+	/** Analyzes loudness of a sound wave and sets LUFS/SamplePeakDB values. 
+	 *  Returns true if a soundwave was analyzed.
+	 */
+	bool AnalyzeSoundWaveLoudness(UObject* Object, bool bMarkDirty = true);
+}
+
 class FWaveformEditorModule : public IWaveformEditorModule
 {
 public:
