[Waveform Editor] AudioQA Editor Win64 Test=UE.EditorAutomation(RunTest=Editor.Python.AudioQA) - * Test 'Editor.Python.AudioQA.test_import_supported_audio_files' failed

PostEditChangeProperty was forking the soundwave when setting the LUFS property on import causing the test to fail saying 2 waves were present.
Switched to using AssetUpdateTags() to updated the content browser metadata.

#Synced-CL 52379869

diff --git a/Engine/Plugins/Editor/WaveformEditor/Source/WaveformEditor/Private/WaveformEditorModule.cpp b/Engine/Plugins/Editor/WaveformEditor/Source/WaveformEditor/Private/WaveformEditorModule.cpp
--- a/Engine/Plugins/Editor/WaveformEditor/Source/WaveformEditor/Private/WaveformEditorModule.cpp
+++ b/Engine/Plugins/Editor/WaveformEditor/Source/WaveformEditor/Private/WaveformEditorModule.cpp
@@ -6,6 +6,8 @@
 #include "Editor.h"
 #include "Modules/ModuleManager.h"
 #include "PropertyEditorModule.h"
+#include "AssetRegistry/IAssetRegistry.h"
+#include "UObject/AssetRegistryTagsContext.h"
 #include "Sound/SoundWave.h"
 #include "Sound/SoundWaveProcedural.h"
 #include "TransformedWaveformViewFactory.h"
@@ -49,10 +51,9 @@ namespace WaveformAnalysis
 
 			SoundWave->SetLoudnessValues(LUFS, SamplePeakDB);
 
-			// Updates the content browser so the tooltip displays updated metadata, uses PostEditChangeProperty() to avoid marking dirty
-			FProperty* Property = USoundWave::StaticClass()->FindPropertyByName(GET_MEMBER_NAME_CHECKED(USoundWave, LUFS));
-			FPropertyChangedEvent PropertyChangedEvent(Property, EPropertyChangeType::ValueSet);
-			SoundWave->PostEditChangeProperty(PropertyChangedEvent);
+			// Refresh the asset registry tags so the content browser tooltip displays updated metadata.
+			// This avoids PostEditChangeProperty which would trigger an unnecessary derived data rebuild.
+			IAssetRegistry::GetChecked().AssetUpdateTags(SoundWave, EAssetRegistryTagsCaller::FullUpdate);
 
 			// Optionally mark the file dirty
 			if (bMarkDirty)
